- Fix CIS

diff --git a/Engine/Plugins/VirtualProduction/CameraCalibration/Source/CameraCalibration/Private/CameraCalibrationSubsystem.cpp b/Engine/Plugins/VirtualProduction/CameraCalibration/Source/CameraCalibration/Private/CameraCalibrationSubsystem.cpp
--- a/Engine/Plugins/VirtualProduction/CameraCalibration/Source/CameraCalibration/Private/CameraCalibrationSubsystem.cpp
+++ b/Engine/Plugins/VirtualProduction/CameraCalibration/Source/CameraCalibration/Private/CameraCalibrationSubsystem.cpp
@@ -8,6 +8,9 @@
 #include "Interfaces/Interface_AssetUserData.h"
 #include "UObject/UObjectIterator.h"
 
+#if WITH_EDITOR
+#include "Editor.h"
+#endif
 
 #define LOCTEXT_NAMESPACE "CameraCalibrationSubsystem"
 
