Fix for crash in FTestCameraCalibration test.

diff --git a/Engine/Plugins/VirtualProduction/CameraCalibration/Source/CameraCalibration/Private/Tests/TestDistortionMapping.cpp b/Engine/Plugins/VirtualProduction/CameraCalibration/Source/CameraCalibration/Private/Tests/TestDistortionMapping.cpp
--- a/Engine/Plugins/VirtualProduction/CameraCalibration/Source/CameraCalibration/Private/Tests/TestDistortionMapping.cpp
+++ b/Engine/Plugins/VirtualProduction/CameraCalibration/Source/CameraCalibration/Private/Tests/TestDistortionMapping.cpp
@@ -19,6 +19,7 @@ namespace CameraCalibrationTestUtil
 		FDistortionMapPoint TempPoint;
 		TempPoint.Focus = 0.0f;
 		TempPoint.Zoom = 0.0f;
+		TempPoint.DistortionInfo.Parameters.AddZeroed();
 
 		float MaxFocus = 100.0f;
 		float MaxZoom = 100.0f;
