CameraCalibration: Fixed a few UI aesthetic issues.

Fixed appearance of more info button and add checkerboard button.
Fixed appearance of section titles.

diff --git a/Engine/Plugins/VirtualProduction/CameraCalibration/Source/CameraCalibrationEditor/Private/AssetEditor/SCameraCalibrationSteps.cpp b/Engine/Plugins/VirtualProduction/CameraCalibration/Source/CameraCalibrationEditor/Private/AssetEditor/SCameraCalibrationSteps.cpp
--- a/Engine/Plugins/VirtualProduction/CameraCalibration/Source/CameraCalibrationEditor/Private/AssetEditor/SCameraCalibrationSteps.cpp
+++ b/Engine/Plugins/VirtualProduction/CameraCalibration/Source/CameraCalibrationEditor/Private/AssetEditor/SCameraCalibrationSteps.cpp
@@ -87,7 +87,7 @@ void SCameraCalibrationSteps::Construct(const FArguments& InArgs, TWeakPtr<FCame
 				[
 					SNew(SBorder) // Background color for title
 					.BorderImage(FEditorStyle::GetBrush("DetailsView.CategoryTop"))
-					.BorderBackgroundColor(FLinearColor(.6, .6, .6, 1.0f))
+					.BorderBackgroundColor(FLinearColor::White)
 					.VAlign(EVerticalAlignment::VAlign_Center)
 					[
 						SNew(SOverlay) 
@@ -96,8 +96,9 @@ void SCameraCalibrationSteps::Construct(const FArguments& InArgs, TWeakPtr<FCame
 						[
 							SNew(STextBlock) // Title text
 							.Text(LOCTEXT("ViewportSettings", "Viewport Settings"))
-							.Font(FEditorStyle::GetFontStyle("DetailsView.CategoryFontStyle"))
-							.ShadowOffset(FVector2D(1.0f, 1.0f))
+							.TransformPolicy(ETextTransformPolicy::ToUpper)
+							.Font(FAppStyle::Get().GetFontStyle(TEXT("PropertyWindow.BoldFont")))
+							.TextStyle(FAppStyle::Get(), "DetailsView.CategoryTextStyle")
 						]
 					]
 				]
@@ -126,7 +127,7 @@ void SCameraCalibrationSteps::Construct(const FArguments& InArgs, TWeakPtr<FCame
 				[
 					SNew(SBorder) // Background color of title
 					.BorderImage(FEditorStyle::GetBrush("DetailsView.CategoryTop"))
-					.BorderBackgroundColor(FLinearColor(.6, .6, .6, 1.0f))
+					.BorderBackgroundColor(FLinearColor::White)
 					.VAlign(EVerticalAlignment::VAlign_Center)
 					[
 						SNew(SOverlay) 
@@ -148,8 +149,9 @@ void SCameraCalibrationSteps::Construct(const FArguments& InArgs, TWeakPtr<FCame
 
 								return LOCTEXT("StepSettings", "Step");
 							})
-							.Font(FEditorStyle::GetFontStyle("DetailsView.CategoryFontStyle"))
-							.ShadowOffset(FVector2D(1.0f, 1.0f))
+							.TransformPolicy(ETextTransformPolicy::ToUpper)
+							.Font(FAppStyle::Get().GetFontStyle(TEXT("PropertyWindow.BoldFont")))
+							.TextStyle(FAppStyle::Get(), "DetailsView.CategoryTextStyle")
 						]
 					]
 				]
diff --git a/Engine/Plugins/VirtualProduction/CameraCalibration/Source/CameraCalibrationEditor/Private/AssetEditor/SLensDistortionToolPanel.cpp b/Engine/Plugins/VirtualProduction/CameraCalibration/Source/CameraCalibrationEditor/Private/AssetEditor/SLensDistortionToolPanel.cpp
--- a/Engine/Plugins/VirtualProduction/CameraCalibration/Source/CameraCalibrationEditor/Private/AssetEditor/SLensDistortionToolPanel.cpp
+++ b/Engine/Plugins/VirtualProduction/CameraCalibration/Source/CameraCalibrationEditor/Private/AssetEditor/SLensDistortionToolPanel.cpp
@@ -188,7 +188,7 @@ TSharedRef<SWidget> SLensDistortionToolPanel::BuildAlgoPickerWidget()
 			.ToolTipText(LOCTEXT("ShowHelp_Tip", "Help about this algo"))
 			.HAlign(HAlign_Center)
 			.VAlign(VAlign_Center)
-			.ButtonColorAndOpacity(FLinearColor::Transparent)
+			.ButtonStyle(FEditorStyle::Get(), "HoverHintOnly")
 			.OnClicked_Lambda([&]() -> FReply
 			{
 				if (!Tool.IsValid())
diff --git a/Engine/Plugins/VirtualProduction/CameraCalibration/Source/CameraCalibrationEditor/Private/AssetEditor/SNodalOffsetToolPanel.cpp b/Engine/Plugins/VirtualProduction/CameraCalibration/Source/CameraCalibrationEditor/Private/AssetEditor/SNodalOffsetToolPanel.cpp
--- a/Engine/Plugins/VirtualProduction/CameraCalibration/Source/CameraCalibrationEditor/Private/AssetEditor/SNodalOffsetToolPanel.cpp
+++ b/Engine/Plugins/VirtualProduction/CameraCalibration/Source/CameraCalibrationEditor/Private/AssetEditor/SNodalOffsetToolPanel.cpp
@@ -187,7 +187,7 @@ TSharedRef<SWidget> SNodalOffsetToolPanel::BuildNodalOffsetAlgoPickerWidget()
 			.ToolTipText(LOCTEXT("ShowHelp_Tip", "Help about this algo"))
 			.HAlign(HAlign_Center)
 			.VAlign(VAlign_Center)
-			.ButtonColorAndOpacity(FLinearColor::Transparent)
+			.ButtonStyle(FEditorStyle::Get(), "HoverHintOnly")
 			.OnClicked_Lambda([&]() -> FReply
 			{
 				if (!NodalOffsetTool.IsValid())
diff --git a/Engine/Plugins/VirtualProduction/CameraCalibration/Source/CameraCalibrationEditor/Private/Calibrators/CameraLensDistortionAlgoCheckerboard.cpp b/Engine/Plugins/VirtualProduction/CameraCalibration/Source/CameraCalibrationEditor/Private/Calibrators/CameraLensDistortionAlgoCheckerboard.cpp
--- a/Engine/Plugins/VirtualProduction/CameraCalibration/Source/CameraCalibrationEditor/Private/Calibrators/CameraLensDistortionAlgoCheckerboard.cpp
+++ b/Engine/Plugins/VirtualProduction/CameraCalibration/Source/CameraCalibrationEditor/Private/Calibrators/CameraLensDistortionAlgoCheckerboard.cpp
@@ -700,7 +700,7 @@ TSharedRef<SWidget> UCameraLensDistortionAlgoCheckerboard::BuildCalibrationDevic
 			.Text(LOCTEXT("Spawn", "Spawn"))
 			.HAlign(HAlign_Center)
 			.VAlign(VAlign_Center)
-			.ButtonColorAndOpacity(FLinearColor::Transparent)
+			.ButtonStyle(FEditorStyle::Get(), "HoverHintOnly")
 			.OnClicked_Lambda([&]() -> FReply
 			{
 				const FCameraCalibrationStepsController* StepsController = GetStepsController();
diff --git a/Engine/Plugins/VirtualProduction/CameraCalibration/Source/CameraCalibrationEditor/Private/Calibrators/CameraNodalOffsetAlgoPoints.cpp b/Engine/Plugins/VirtualProduction/CameraCalibration/Source/CameraCalibrationEditor/Private/Calibrators/CameraNodalOffsetAlgoPoints.cpp
--- a/Engine/Plugins/VirtualProduction/CameraCalibration/Source/CameraCalibrationEditor/Private/Calibrators/CameraNodalOffsetAlgoPoints.cpp
+++ b/Engine/Plugins/VirtualProduction/CameraCalibration/Source/CameraCalibrationEditor/Private/Calibrators/CameraNodalOffsetAlgoPoints.cpp
@@ -1048,9 +1048,12 @@ TSharedRef<SWidget> UCameraNodalOffsetAlgoPoints::BuildCalibrationDevicePickerWi
 
 TSharedRef<SWidget> UCameraNodalOffsetAlgoPoints::BuildCalibrationActionButtons()
 {
+	const float ButtonPadding = 3.0f;
+
 	return SNew(SVerticalBox)
 
 		+ SVerticalBox::Slot() // Row manipulation
+		.AutoHeight()
 		[
 			SNew(SHorizontalBox)
 
@@ -1075,6 +1078,8 @@ TSharedRef<SWidget> UCameraNodalOffsetAlgoPoints::BuildCalibrationActionButtons(
 			.MaxDesiredHeight(0.5 * FCameraCalibrationWidgetHelpers::DefaultRowHeight)
 		]
 		+ SVerticalBox::Slot() // Apply To Calibrator
+		.AutoHeight()
+		.Padding(0, ButtonPadding)
 		[
 			SNew(SButton)
 			.Text(LOCTEXT("ApplyToCalibrator", "Apply To Calibrator"))
@@ -1088,6 +1093,8 @@ TSharedRef<SWidget> UCameraNodalOffsetAlgoPoints::BuildCalibrationActionButtons(
 			})
 		]
 		+ SVerticalBox::Slot() // Apply To Camera Parent
+		.AutoHeight()
+		.Padding(0, ButtonPadding)
 		[
 			SNew(SButton)
 			.Text(LOCTEXT("ApplyToTrackingOrigin", "Apply To Camera Parent"))
@@ -1101,18 +1108,20 @@ TSharedRef<SWidget> UCameraNodalOffsetAlgoPoints::BuildCalibrationActionButtons(
 			})
 		]
 		+ SVerticalBox::Slot() // Apply To Calibrator Parent
-			[
-				SNew(SButton)
-				.Text(LOCTEXT("ApplyToCalibratorParent", "Apply To Calibrator Parent"))
-				.HAlign(HAlign_Center)
-				.VAlign(VAlign_Center)
-				.OnClicked_Lambda([&]() -> FReply
-				{
+		.AutoHeight()
+		.Padding(0, ButtonPadding)
+		[
+			SNew(SButton)
+			.Text(LOCTEXT("ApplyToCalibratorParent", "Apply To Calibrator Parent"))
+			.HAlign(HAlign_Center)
+			.VAlign(VAlign_Center)
+			.OnClicked_Lambda([&]() -> FReply
+			{
 				FScopedTransaction Transaction(LOCTEXT("ApplyNodalOffsetToCalibratorParent", "Applying Nodal Offset to Calibrator Parent"));
-					ApplyNodalOffsetToCalibratorParent();
-					return FReply::Handled();
-				})
-			]
+				ApplyNodalOffsetToCalibratorParent();
+				return FReply::Handled();
+			})
+		]
 		;
 }
 
