CameraCalibration/LiveLink - Miscellaneous UI fixes
Replaced font awesome glyphs in lens file editor with proper icons
Cleaned up alignment issues in lens file add point dialog
Clarified confusing text in lens file remove point dialog
Refactored camera calibration editor style to inherit from FSlateStyleSet
Fixed mixed case issue in livelink controller hierarchy details customziation
Fixed bad foreground color and style for live link subject representation picker
#rnx

diff --git a/Engine/Plugins/Animation/LiveLink/Source/LiveLinkEditor/Private/LiveLinkComponentDetailCustomization.cpp b/Engine/Plugins/Animation/LiveLink/Source/LiveLinkEditor/Private/LiveLinkComponentDetailCustomization.cpp
--- a/Engine/Plugins/Animation/LiveLink/Source/LiveLinkEditor/Private/LiveLinkComponentDetailCustomization.cpp
+++ b/Engine/Plugins/Animation/LiveLink/Source/LiveLinkEditor/Private/LiveLinkComponentDetailCustomization.cpp
@@ -366,7 +366,7 @@ FText FLiveLinkComponentDetailCustomization::HandleControllerStatusToolTipText(T
 
 TSharedRef<SWidget> FLiveLinkComponentDetailCustomization::BuildControllerNameWidget(TSharedPtr<IPropertyHandle> ControllersProperty, TSubclassOf<ULiveLinkRole> RoleClass) const
 {
-	return ControllersProperty->CreatePropertyNameWidget(RoleClass.Get()->GetDisplayNameText());
+	return ControllersProperty->CreatePropertyNameWidget(RoleClass.Get()->GetDisplayNameText().ToUpper());
 }
 
 TSharedRef<SWidget> FLiveLinkComponentDetailCustomization::BuildControllerValueWidget(TSharedPtr<IPropertyHandle> RoleKeyPropertyHandle, TSubclassOf<ULiveLinkRole> RoleClass, const FText& ControllerName) const
@@ -375,20 +375,17 @@ TSharedRef<SWidget> FLiveLinkComponentDetailCustomization::BuildControllerValueW
 		SNew(SHorizontalBox)
 		+ SHorizontalBox::Slot()
 		.AutoWidth()
-		.Padding(0, 0, 0, 0)
-		.HAlign(HAlign_Left)
-		.VAlign(VAlign_Center)
-		[
-			SNew(STextBlock)
-			.Text(ControllerName)
-		]
-		+ SHorizontalBox::Slot()
-		.AutoWidth()
 		.Padding(FMargin(4.0f, 0.0f, 0.0f, 0.0f))
 		[
 			SNew(SComboButton)
 			.OnGetMenuContent(this, &FLiveLinkComponentDetailCustomization::HandleControllerComboButton, RoleKeyPropertyHandle)
 			.ContentPadding(FMargin(4.0, 2.0))
+			.ButtonContent()
+			[
+				SNew(STextBlock)
+				.Text(ControllerName)
+				.Font(FAppStyle::Get().GetFontStyle("SmallFont"))
+			]
 		]
 		+ SHorizontalBox::Slot()
 		.AutoWidth()
diff --git a/Engine/Plugins/Animation/LiveLink/Source/LiveLinkEditor/Private/LiveLinkEditorModule.cpp b/Engine/Plugins/Animation/LiveLink/Source/LiveLinkEditor/Private/LiveLinkEditorModule.cpp
--- a/Engine/Plugins/Animation/LiveLink/Source/LiveLinkEditor/Private/LiveLinkEditorModule.cpp
+++ b/Engine/Plugins/Animation/LiveLink/Source/LiveLinkEditor/Private/LiveLinkEditorModule.cpp
@@ -126,6 +126,17 @@ public:
 			.SetMenuBorderPadding(FMargin(0.0f));
 		StyleSet->Set("ComboButton", ComboButton);
 
+		FEditableTextBoxStyle EditableTextStyle = FEditableTextBoxStyle()
+			.SetFont(FCoreStyle::GetDefaultFontStyle("Regular", 9))
+			.SetBackgroundImageNormal(FSlateNoResource())
+			.SetBackgroundImageHovered(FSlateNoResource())
+			.SetBackgroundImageFocused(FSlateNoResource())
+			.SetBackgroundImageReadOnly(FSlateNoResource())
+			.SetBackgroundColor(FLinearColor::Transparent)
+			.SetForegroundColor(FSlateColor::UseForeground());
+
+		StyleSet->Set("EditableTextBox", EditableTextStyle);
+
 		FSlateStyleRegistry::RegisterSlateStyle(*StyleSet.Get());
 
 		RegisterSettings();
diff --git a/Engine/Plugins/Animation/LiveLink/Source/LiveLinkEditor/Private/SLiveLinkSubjectRepresentationPicker.cpp b/Engine/Plugins/Animation/LiveLink/Source/LiveLinkEditor/Private/SLiveLinkSubjectRepresentationPicker.cpp
--- a/Engine/Plugins/Animation/LiveLink/Source/LiveLinkEditor/Private/SLiveLinkSubjectRepresentationPicker.cpp
+++ b/Engine/Plugins/Animation/LiveLink/Source/LiveLinkEditor/Private/SLiveLinkSubjectRepresentationPicker.cpp
@@ -183,6 +183,7 @@ void SLiveLinkSubjectRepresentationPicker::Construct(const FArguments& InArgs)
 			.SelectAllTextOnCommit(true)
 			.ClearKeyboardFocusOnCommit(false)
 			.Font(InArgs._Font)
+			.Style(FLiveLinkEditorPrivate::GetStyleSet(), "EditableTextBox")
 		];
 	}
 
@@ -191,9 +192,9 @@ void SLiveLinkSubjectRepresentationPicker::Construct(const FArguments& InArgs)
 		SAssignNew(PickerComboButton, SComboButton)
 		.ComboButtonStyle(InArgs._ComboButtonStyle)
 		.ButtonStyle(InArgs._ButtonStyle)
-		.ForegroundColor(InArgs._ForegroundColor)
+		.ForegroundColor(FSlateColor::UseForeground())
 		.ContentPadding(InArgs._ContentPadding)
-		.VAlign(VAlign_Fill)
+		.VAlign(VAlign_Center)
 		.OnGetMenuContent(this, &SLiveLinkSubjectRepresentationPicker::BuildMenu)
 		.ButtonContent()
 		[
diff --git a/Engine/Plugins/VirtualProduction/CameraCalibration/Source/CameraCalibrationEditor/Private/AssetEditor/CameraCalibrationCommands.h b/Engine/Plugins/VirtualProduction/CameraCalibration/Source/CameraCalibrationEditor/Private/AssetEditor/CameraCalibrationCommands.h
--- a/Engine/Plugins/VirtualProduction/CameraCalibration/Source/CameraCalibrationEditor/Private/AssetEditor/CameraCalibrationCommands.h
+++ b/Engine/Plugins/VirtualProduction/CameraCalibration/Source/CameraCalibrationEditor/Private/AssetEditor/CameraCalibrationCommands.h
@@ -15,7 +15,7 @@ class FCameraCalibrationCommands : public TCommands<FCameraCalibrationCommands>
 public:
 
 	FCameraCalibrationCommands()
-		: TCommands<FCameraCalibrationCommands>(TEXT("CameraCalibration"), NSLOCTEXT("CameraCalibrationCommands", "CameraCalibrationCommandsEdit", "CameraCalibration"), NAME_None, FCameraCalibrationEditorStyle::GetStyleSetName())
+		: TCommands<FCameraCalibrationCommands>(TEXT("CameraCalibration"), NSLOCTEXT("CameraCalibrationCommands", "CameraCalibrationCommandsEdit", "CameraCalibration"), NAME_None, FCameraCalibrationEditorStyle::Get().GetStyleSetName())
 	{}
 	
 	/**
diff --git a/Engine/Plugins/VirtualProduction/CameraCalibration/Source/CameraCalibrationEditor/Private/AssetEditor/SCameraCalibrationCurveEditorPanel.cpp b/Engine/Plugins/VirtualProduction/CameraCalibration/Source/CameraCalibrationEditor/Private/AssetEditor/SCameraCalibrationCurveEditorPanel.cpp
--- a/Engine/Plugins/VirtualProduction/CameraCalibration/Source/CameraCalibrationEditor/Private/AssetEditor/SCameraCalibrationCurveEditorPanel.cpp
+++ b/Engine/Plugins/VirtualProduction/CameraCalibration/Source/CameraCalibrationEditor/Private/AssetEditor/SCameraCalibrationCurveEditorPanel.cpp
@@ -116,6 +116,7 @@ TSharedPtr<FExtender> SCameraCalibrationCurveEditorPanel::GetToolbarExtender()
 		static void FillToolbar(FToolBarBuilder& ToolBarBuilder, TSharedRef<SCurveKeyDetailPanel> InKeyDetailsPanel, TSharedRef<SCameraCalibrationCurveEditorPanel> InEditorPanel)
 		{
 			ToolBarBuilder.BeginSection("View");
+			ToolBarBuilder.BeginStyleOverride("AssetEditorToolbar");
 			{
 				// Dropdown Menu for choosing your viewing mode
 				TAttribute<FSlateIcon> ViewModeIcon;
diff --git a/Engine/Plugins/VirtualProduction/CameraCalibration/Source/CameraCalibrationEditor/Private/AssetEditor/SCameraCalibrationCurveEditorView.cpp b/Engine/Plugins/VirtualProduction/CameraCalibration/Source/CameraCalibrationEditor/Private/AssetEditor/SCameraCalibrationCurveEditorView.cpp
--- a/Engine/Plugins/VirtualProduction/CameraCalibration/Source/CameraCalibrationEditor/Private/AssetEditor/SCameraCalibrationCurveEditorView.cpp
+++ b/Engine/Plugins/VirtualProduction/CameraCalibration/Source/CameraCalibrationEditor/Private/AssetEditor/SCameraCalibrationCurveEditorView.cpp
@@ -278,9 +278,9 @@ void SCameraCalibrationCurveEditorView::PaintView(const FPaintArgs& Args, const
 			const FLensDataCurveModel* const LensDataCurveModel = static_cast<const FLensDataCurveModel* const>(Curve);
 
 			// Draw an axis descriptor for the current curve editor view
-			const FText AxisDescriptor = FText::Format(LOCTEXT("CurveEditorTooltipName", "X-Axis: {0}   Y-Axis: {1} {2}"), LensDataCurveModel->GetKeyLabel(), LensDataCurveModel->GetLongDisplayName(), LensDataCurveModel->GetValueLabel());
+			const FText AxisDescriptor = FText::Format(LOCTEXT("CurveEditorAxisDescriptor", "X-Axis: {0}   Y-Axis: {1} {2}"), LensDataCurveModel->GetKeyLabel(), LensDataCurveModel->GetLongDisplayName(), LensDataCurveModel->GetValueLabel());
 
-			const FSlateFontInfo FontInfo = FCoreStyle::Get().GetFontStyle("FontAwesome.11");
+			const FSlateFontInfo FontInfo = FAppStyle::Get().GetFontStyle("CurveEd.LabelFont");
 			const TSharedRef<FSlateFontMeasure> FontMeasure = FSlateApplication::Get().GetRenderer()->GetFontMeasureService();
 			const ESlateDrawEffect DrawEffects = ShouldBeEnabled(bParentEnabled) ? ESlateDrawEffect::None : ESlateDrawEffect::DisabledEffect;
 
diff --git a/Engine/Plugins/VirtualProduction/CameraCalibration/Source/CameraCalibrationEditor/Private/AssetEditor/SCameraCalibrationRemovePointDialog.cpp b/Engine/Plugins/VirtualProduction/CameraCalibration/Source/CameraCalibrationEditor/Private/AssetEditor/SCameraCalibrationRemovePointDialog.cpp
--- a/Engine/Plugins/VirtualProduction/CameraCalibration/Source/CameraCalibrationEditor/Private/AssetEditor/SCameraCalibrationRemovePointDialog.cpp
+++ b/Engine/Plugins/VirtualProduction/CameraCalibration/Source/CameraCalibrationEditor/Private/AssetEditor/SCameraCalibrationRemovePointDialog.cpp
@@ -348,7 +348,7 @@ void SCameraCalibrationRemovePointDialog::Construct(const FArguments& InArgs, co
 		return;
 	}
 	
-	const FText DialogText = FText::Format(LOCTEXT("DialogTextLabel", "Select points to remove which have same raw input than {0}"), FText::FromName(BaseDataTable->GetScriptStruct()->GetFName()));
+	const FText DialogText = LOCTEXT("DialogTextLabel", "The calibration data you wish to delete may be inherently linked to additional data.\nChoose any and all linked data you wish to delete.");
 
 	const TSharedPtr<SWidget> ButtonsWidget = [this]()
 	{
diff --git a/Engine/Plugins/VirtualProduction/CameraCalibration/Source/CameraCalibrationEditor/Private/AssetEditor/SLensDataAddPointDialog.cpp b/Engine/Plugins/VirtualProduction/CameraCalibration/Source/CameraCalibrationEditor/Private/AssetEditor/SLensDataAddPointDialog.cpp
--- a/Engine/Plugins/VirtualProduction/CameraCalibration/Source/CameraCalibrationEditor/Private/AssetEditor/SLensDataAddPointDialog.cpp
+++ b/Engine/Plugins/VirtualProduction/CameraCalibration/Source/CameraCalibrationEditor/Private/AssetEditor/SLensDataAddPointDialog.cpp
@@ -136,37 +136,40 @@ void SLensDataAddPointDialog::Construct(const FArguments& InArgs, ULensFile* InL
 			[
 				SNew(SHorizontalBox)
 				+ SHorizontalBox::Slot()
+				.VAlign(VAlign_Center)
+				.Padding(FMargin(4.f, 0.f, 0.f, 0.f))
 				.AutoWidth()
 				[
 					SNew(STextBlock)
-					.Text(LOCTEXT("DataCategoryLabel", "Data Type "))
+					.Text(LOCTEXT("DataCategoryLabel", "Data Type"))
+					.Font(FAppStyle::Get().GetFontStyle("NormalFontBold"))
+					.TransformPolicy(ETextTransformPolicy::ToUpper)
 				]
 				+ SHorizontalBox::Slot()
-				.HAlign(HAlign_Fill)
-				[
-					SNew(STextBlock)
-					.Text(MakeAttributeLambda([=]
-					{
-							switch (SelectedCategory)
-							{
-								case ELensDataCategory::Focus: return LOCTEXT("FocusCategory", "Encoder - Focus");
-								case ELensDataCategory::Iris: return LOCTEXT("IrisCategory", "Encoder - Iris");
-								case ELensDataCategory::Zoom: return LOCTEXT("ZoomCategory", "Focal Length");
-								case ELensDataCategory::Distortion: return LOCTEXT("DistortionCategory", "Distortion Parameters");
-								case ELensDataCategory::ImageCenter: return LOCTEXT("ImageCenterCategory", "Image Center");
-								case ELensDataCategory::NodalOffset: return LOCTEXT("NodalOffsetCategory", "Nodal offset");
-								case ELensDataCategory::STMap: return LOCTEXT("STMapCategory", "STMap");
-								default: return LOCTEXT("InvalidCategory", "Invalid");
-							}
-					}))
-				]
-				+ SHorizontalBox::Slot()
-				.AutoWidth()
+				.HAlign(HAlign_Right)
 				.Padding(FMargin(4.0f, 0.0f, 0.0f, 0.0f))
 				[
 					SNew(SComboButton)
 					.OnGetMenuContent(this, &SLensDataAddPointDialog::MakeDataCategoryMenuWidget)
 					.ContentPadding(FMargin(4.0, 2.0))
+					.ButtonContent()
+					[
+						SNew(STextBlock)
+						.Text(MakeAttributeLambda([=]
+						{
+							switch (SelectedCategory)
+							{
+							case ELensDataCategory::Focus: return LOCTEXT("FocusCategory", "Encoder - Focus");
+							case ELensDataCategory::Iris: return LOCTEXT("IrisCategory", "Encoder - Iris");
+							case ELensDataCategory::Zoom: return LOCTEXT("ZoomCategory", "Focal Length");
+							case ELensDataCategory::Distortion: return LOCTEXT("DistortionCategory", "Distortion Parameters");
+							case ELensDataCategory::ImageCenter: return LOCTEXT("ImageCenterCategory", "Image Center");
+							case ELensDataCategory::NodalOffset: return LOCTEXT("NodalOffsetCategory", "Nodal offset");
+							case ELensDataCategory::STMap: return LOCTEXT("STMapCategory", "STMap");
+							default: return LOCTEXT("InvalidCategory", "Invalid");
+							}
+						}))
+					]
 				]
 			]
 		]
@@ -251,20 +254,23 @@ TSharedRef<SWidget> SLensDataAddPointDialog::MakeTrackingDataWidget()
 		return SNew(SHorizontalBox)
 			+ SHorizontalBox::Slot()
 			.AutoWidth()
+			.Padding(0.f, 8.f, 0.f, 8.f)
 			.VAlign(VAlign_Center)
 			[
 				SNew(SCheckBox)
 				.IsChecked(MakeAttributeLambda([this, Index]() { return IsTrackingDataOverrided(Index) ? ECheckBoxState::Checked : ECheckBoxState::Unchecked; }))
 				.OnCheckStateChanged(this, &SLensDataAddPointDialog::OnOverrideTrackingData, Index)
+				.ToolTipText(LOCTEXT("TrackingDataCheckboxTooltip", "Check to override incoming tracking data for this input"))
 			]
 			+ SHorizontalBox::Slot()
-			.FillWidth(0.2f)
+			.VAlign(VAlign_Center)
+			.Padding(4.f, 0.f, 0.f, 0.f)
 			[
 				SNew(STextBlock)
 				.Text(FText::FromString(Label))
 			]
 			+ SHorizontalBox::Slot()
-			.FillWidth(0.8f)
+			.VAlign(VAlign_Center)
 			[
 				SNew(SNumericEntryBox<float>)
 				.IsEnabled(MakeAttributeLambda([this, Index]() { return IsTrackingDataOverrided(Index); }))
@@ -296,12 +302,10 @@ TSharedRef<SWidget> SLensDataAddPointDialog::MakeTrackingDataWidget()
 			TrackingWidget = 
 				SNew(SVerticalBox)
 				+ SVerticalBox::Slot()
-				.AutoHeight()
 				[
 					MakeRowWidget(TEXT("Input Focus"), 0)
 				]
 				+ SVerticalBox::Slot()
-				.AutoHeight()
 				[
 					MakeRowWidget(TEXT("Input Zoom"), 1)
 				];
diff --git a/Engine/Plugins/VirtualProduction/CameraCalibration/Source/CameraCalibrationEditor/Private/AssetEditor/SLensDataListItem.cpp b/Engine/Plugins/VirtualProduction/CameraCalibration/Source/CameraCalibrationEditor/Private/AssetEditor/SLensDataListItem.cpp
--- a/Engine/Plugins/VirtualProduction/CameraCalibration/Source/CameraCalibrationEditor/Private/AssetEditor/SLensDataListItem.cpp
+++ b/Engine/Plugins/VirtualProduction/CameraCalibration/Source/CameraCalibrationEditor/Private/AssetEditor/SLensDataListItem.cpp
@@ -6,6 +6,7 @@
 #include "ScopedTransaction.h"
 #include "SLensDataEditPointDialog.h"
 #include "UI/CameraCalibrationEditorStyle.h"
+#include "Widgets/Images/SImage.h"
 #include "Widgets/Input/SButton.h"
 #include "SCameraCalibrationRemovePointDialog.h"
 
@@ -230,12 +231,11 @@ void SLensDataItem::Construct(const FArguments& InArgs, const TSharedRef<STableV
 				.OnClicked(this, &SLensDataItem::OnEditPointClicked)
 				.IsEnabled(InArgs._AllowEditPoint)
 				.Visibility(InArgs._EditPointVisibility)
-				.ButtonStyle(FEditorStyle::Get(), "FlatButton")
+				.ButtonStyle(FAppStyle::Get(), "FlatButton")
+				.ToolTipText(LOCTEXT("EditLensDataPoint", "Edit the value at this point"))
 				[
-					SNew(STextBlock)
-					.TextStyle(FCameraCalibrationEditorStyle::Get(), "CameraCalibration.Button.TextStyle")
-					.Font(FEditorStyle::Get().GetFontStyle("FontAwesome.10"))
-					.Text(FText::FromString(FString(TEXT("\xf044"))) /*fa-edit*/)
+					SNew(SImage)
+					.Image(FAppStyle::Get().GetBrush("Icons.Edit"))
 				]
 			]
 			+ SHorizontalBox::Slot()
@@ -244,12 +244,11 @@ void SLensDataItem::Construct(const FArguments& InArgs, const TSharedRef<STableV
 				SNew(SButton)
 				.IsEnabled(InArgs._AllowRemoval)
 				.OnClicked(this, &SLensDataItem::OnRemovePointClicked)
-				.ButtonStyle(FCameraCalibrationEditorStyle::Get(), "CameraCalibration.RemoveButton")
+				.ButtonStyle(FEditorStyle::Get(), "FlatButton")
+				.ToolTipText(LOCTEXT("RemoveLensDataPoint", "Remove this point"))
 				[
-					SNew(STextBlock)
-					.TextStyle(FCameraCalibrationEditorStyle::Get(), "CameraCalibration.Button.TextStyle")
-					.Font(FEditorStyle::Get().GetFontStyle("FontAwesome.10"))
-					.Text(FText::FromString(FString(TEXT("\xf00d"))) /*fa-times*/)
+					SNew(SImage)
+					.Image(FAppStyle::Get().GetBrush("Icons.Delete"))
 				]
 			]
 		], OwnerTable);
diff --git a/Engine/Plugins/VirtualProduction/CameraCalibration/Source/CameraCalibrationEditor/Private/AssetEditor/SLensDataViewer.cpp b/Engine/Plugins/VirtualProduction/CameraCalibration/Source/CameraCalibrationEditor/Private/AssetEditor/SLensDataViewer.cpp
--- a/Engine/Plugins/VirtualProduction/CameraCalibration/Source/CameraCalibrationEditor/Private/AssetEditor/SLensDataViewer.cpp
+++ b/Engine/Plugins/VirtualProduction/CameraCalibration/Source/CameraCalibrationEditor/Private/AssetEditor/SLensDataViewer.cpp
@@ -27,6 +27,7 @@
 #include "Curves/LensSTMapCurveModel.h"
 #include "Framework/MultiBox/MultiBoxBuilder.h"
 #include "Misc/MessageDialog.h"
+#include "Widgets/Images/SImage.h"
 #include "Widgets/Input/SButton.h"
 
 
@@ -339,36 +340,39 @@ TSharedRef<SWidget> SLensDataViewer::MakeToolbarWidget(TSharedRef<SCameraCalibra
 {
 	// Curve toolbar
 	FToolBarBuilder ToolBarBuilder(CurvePanel->GetCommands(), FMultiBoxCustomization::None, CurvePanel->GetToolbarExtender(), true);
-	ToolBarBuilder.SetStyle(&FEditorStyle::Get(), "Sequencer.ToolBar");
 	ToolBarBuilder.BeginSection("Asset");
-	
+	ToolBarBuilder.BeginStyleOverride("AssetEditorToolbar");
+
 	TSharedRef<SWidget> AddPointButton =
 			 SNew(SButton)
-			.ButtonStyle(FEditorStyle::Get(), "ToggleButton")
+			.ButtonStyle(FEditorStyle::Get(), "FlatButton")
+			.VAlign(VAlign_Center)
+			.HAlign(HAlign_Center)
+			.ToolTipText(LOCTEXT("AddLensDataPoint", "Add a lens data point"))
 			.OnClicked_Lambda([this]()
 			{
 				OnAddDataPointHandler();
 				return FReply::Handled();
 			})
 			[
-				SNew(STextBlock)
-				.Font(FEditorStyle::Get().GetFontStyle("FontAwesome.14"))
-				.Text(FText::FromString(FString(TEXT("\xf067"))) /*fa-plus*/)
-				.ColorAndOpacity(FLinearColor::White)
+				SNew(SImage)
+				.Image(FAppStyle::Get().GetBrush("Icons.Plus"))
 			];
 
 	TSharedRef<SWidget> ClearAllButton =
 			SNew(SButton)
-			.ButtonStyle(FEditorStyle::Get(), "ToggleButton")
+			.ButtonStyle(FEditorStyle::Get(), "FlatButton")
+			.VAlign(VAlign_Center)
+			.ToolTipText(LOCTEXT("DeleteLensData", "Delete all calibrated lens data"))
 			.OnClicked(this, &SLensDataViewer::OnClearLensFileClicked)
 			[
-				SNew(STextBlock)
-				.Font(FEditorStyle::Get().GetFontStyle("FontAwesome.14"))
-				.Text(FText::FromString(FString(TEXT("\xf1f8"))) /*fa-trash*/)
-				.ColorAndOpacity(FLinearColor::White)
+				SNew(SImage)
+				.Image(FAppStyle::Get().GetBrush("Icons.Delete"))
 			];
 
+	ToolBarBuilder.AddSeparator();
 	ToolBarBuilder.AddWidget(AddPointButton);
+	ToolBarBuilder.AddSeparator();
 	ToolBarBuilder.AddWidget(ClearAllButton);
 		
 	ToolBarBuilder.EndSection();
diff --git a/Engine/Plugins/VirtualProduction/CameraCalibration/Source/CameraCalibrationEditor/Private/CameraCalibrationEditorModule.cpp b/Engine/Plugins/VirtualProduction/CameraCalibration/Source/CameraCalibrationEditor/Private/CameraCalibrationEditorModule.cpp
--- a/Engine/Plugins/VirtualProduction/CameraCalibration/Source/CameraCalibrationEditor/Private/CameraCalibrationEditorModule.cpp
+++ b/Engine/Plugins/VirtualProduction/CameraCalibration/Source/CameraCalibrationEditor/Private/CameraCalibrationEditorModule.cpp
@@ -32,7 +32,7 @@ DEFINE_LOG_CATEGORY(LogCameraCalibrationEditor);
 void FCameraCalibrationEditorModule::StartupModule()
 {
 	FCameraCalibrationCommands::Register();
-	FCameraCalibrationEditorStyle::Register();
+	FCameraCalibrationEditorStyle::Get();
 
 	// Register AssetTypeActions
 	auto RegisterAssetTypeAction = [this](const TSharedRef<IAssetTypeActions>& InAssetTypeAction)
@@ -80,6 +80,7 @@ const FPlacementCategoryInfo* FCameraCalibrationEditorModule::GetVirtualProducti
 	{
 		FPlacementCategoryInfo Info(
 			LOCTEXT("VirtualProductionCategoryName", "Virtual Production"),
+			FSlateIcon(FCameraCalibrationEditorStyle::Get().GetStyleSetName(), "PlacementBrowser.Icons.VirtualProduction"),
 			VirtualProductionName,
 			TEXT("PMVirtualProduction"),
 			25
@@ -180,7 +181,6 @@ void FCameraCalibrationEditorModule::ShutdownModule()
 			}
 		}
 
-		FCameraCalibrationEditorStyle::Unregister();
 		FCameraCalibrationCommands::Unregister();
 
 		UnregisterPlacementModeItems();
diff --git a/Engine/Plugins/VirtualProduction/CameraCalibration/Source/CameraCalibrationEditor/Private/UI/CameraCalibrationEditorStyle.cpp b/Engine/Plugins/VirtualProduction/CameraCalibration/Source/CameraCalibrationEditor/Private/UI/CameraCalibrationEditorStyle.cpp
deleted file mode 100644
diff --git a/Engine/Plugins/VirtualProduction/CameraCalibration/Source/CameraCalibrationEditor/Private/UI/CameraCalibrationEditorStyle.h b/Engine/Plugins/VirtualProduction/CameraCalibration/Source/CameraCalibrationEditor/Private/UI/CameraCalibrationEditorStyle.h
--- a/Engine/Plugins/VirtualProduction/CameraCalibration/Source/CameraCalibrationEditor/Private/UI/CameraCalibrationEditorStyle.h
+++ b/Engine/Plugins/VirtualProduction/CameraCalibration/Source/CameraCalibrationEditor/Private/UI/CameraCalibrationEditorStyle.h
@@ -1,20 +1,48 @@
 // Copyright Epic Games, Inc. All Rights Reserved.
 
 #pragma once
-#include "CoreMinimal.h"
 
-class ISlateStyle;
+#include "CoreMinimal.h"
+#include "Styling/SlateStyleRegistry.h"
+#include "Styling/SlateStyleMacros.h"
+#include "Styling/SlateStyle.h"
 
-//////////////////////////////////////////////////////////////////////////
-// FCameraCalibrationEditorStyle
 
-class FCameraCalibrationEditorStyle
+/**
+ * Implements the visual style of the camera calibration tools UI.
+ */
+class FCameraCalibrationEditorStyle	: public FSlateStyleSet
 {
 public:
-	static void Register();
-	static void Unregister();
 
-	static FName GetStyleSetName();
+	/** Default constructor. */
+	FCameraCalibrationEditorStyle()	: FSlateStyleSet("CameraCalibrationEditorStyle")
+	{
+		const FVector2D Icon16x16(16.0f, 16.0f);
+		const FVector2D Icon20x20(20.0f, 20.0f);
+		const FVector2D Icon64x64(64.0f, 64.0f);
+
+		// Set placement browser icons
+		SetContentRoot(FPaths::EngineContentDir() / TEXT("Editor/Slate/"));
+		Set("PlacementBrowser.Icons.VirtualProduction", new IMAGE_BRUSH_SVG("Starship/Common/VirtualProduction", Icon16x16));
+
+		// Set miscellaneous icons
+		SetContentRoot(FPaths::EnginePluginsDir() / TEXT("VirtualProduction/CameraCalibration/Content/Editor/Icons/"));
+		Set("ClassThumbnail.LensFile", new IMAGE_BRUSH("LensFileIcon_64x", Icon64x64));
+		Set("ClassIcon.LensFile", new IMAGE_BRUSH("LensFileIcon_20x", Icon20x20));
+
+		FSlateStyleRegistry::RegisterSlateStyle(*this);
+	}
+
+	/** Virtual destructor. */
+	virtual ~FCameraCalibrationEditorStyle()
+	{
+		FSlateStyleRegistry::UnRegisterSlateStyle(*this);
+	}
 
-	static const ISlateStyle& Get();
-};
\ No newline at end of file
+	static FCameraCalibrationEditorStyle& Get()
+	{
+		static FCameraCalibrationEditorStyle Inst;
+		return Inst;
+	}
+};
diff --git a/Engine/Plugins/VirtualProduction/CameraCalibration/Source/CameraCalibrationEditor/Private/UI/CameraCalibrationMenuEntry.cpp b/Engine/Plugins/VirtualProduction/CameraCalibration/Source/CameraCalibrationEditor/Private/UI/CameraCalibrationMenuEntry.cpp
--- a/Engine/Plugins/VirtualProduction/CameraCalibration/Source/CameraCalibrationEditor/Private/UI/CameraCalibrationMenuEntry.cpp
+++ b/Engine/Plugins/VirtualProduction/CameraCalibration/Source/CameraCalibrationEditor/Private/UI/CameraCalibrationMenuEntry.cpp
@@ -128,7 +128,7 @@ struct FCameraCalibrationMenuEntryImpl
 				NAME_None,
 				LOCTEXT("LensFile_Label", "Lens File"),
 				MakeAttributeLambda(TooltipLambda),
-				FSlateIcon(FCameraCalibrationEditorStyle::GetStyleSetName(), TEXT("ToolbarIcon.LensFile"))
+				FSlateIcon(FCameraCalibrationEditorStyle::Get().GetStyleSetName(), TEXT("ToolbarIcon.LensFile"))
 			);
 
 			// Add a simple drop-down menu (no label, no icon for the drop-down button itself) that list the lens files available
@@ -154,7 +154,7 @@ struct FCameraCalibrationMenuEntryImpl
 			MenuBuilder.AddMenuEntry(
 				LOCTEXT("CreateMenuLabel", "New Empty Lens File"),
 				LOCTEXT("CreateMenuTooltip", "Create a new Lens File asset."),
-				FSlateIcon(FCameraCalibrationEditorStyle::GetStyleSetName(), TEXT("ClassIcon.LensFile")),
+				FSlateIcon(FCameraCalibrationEditorStyle::Get().GetStyleSetName(), TEXT("ClassIcon.LensFile")),
 				FUIAction(
 					FExecuteAction::CreateRaw(this, &FCameraCalibrationMenuEntryImpl::CreateNewLensFile)
 				)
diff --git a/Engine/Plugins/VirtualProduction/CameraCalibrationCore/Source/CameraCalibrationCore/Private/Tables/BaseLensTable.cpp b/Engine/Plugins/VirtualProduction/CameraCalibrationCore/Source/CameraCalibrationCore/Private/Tables/BaseLensTable.cpp
--- a/Engine/Plugins/VirtualProduction/CameraCalibrationCore/Source/CameraCalibrationCore/Private/Tables/BaseLensTable.cpp
+++ b/Engine/Plugins/VirtualProduction/CameraCalibrationCore/Source/CameraCalibrationCore/Private/Tables/BaseLensTable.cpp
@@ -7,11 +7,11 @@ FName FBaseLensTable::GetFriendlyPointName(ELensDataCategory InCategory)
 {
 	switch (InCategory)
 	{
-		case ELensDataCategory::Zoom: return TEXT("Focal Length Point");
-		case ELensDataCategory::Distortion: return TEXT("Distortion Point");
-		case ELensDataCategory::ImageCenter: return TEXT("Image Center Point");
-		case ELensDataCategory::STMap: return TEXT("ST Map Point");
-		case ELensDataCategory::NodalOffset: return TEXT("Nodal Offset Point");
+		case ELensDataCategory::Zoom: return TEXT("Focal Length");
+		case ELensDataCategory::Distortion: return TEXT("Distortion Parameters");
+		case ELensDataCategory::ImageCenter: return TEXT("Image Center");
+		case ELensDataCategory::STMap: return TEXT("ST Map");
+		case ELensDataCategory::NodalOffset: return TEXT("Nodal Offset");
 	}
 
 	return TEXT("");
diff --git a/Engine/Plugins/VirtualProduction/CameraCalibrationCore/Source/CameraCalibrationCoreEditor/Private/CameraCalibrationCoreEditorModule.cpp b/Engine/Plugins/VirtualProduction/CameraCalibrationCore/Source/CameraCalibrationCoreEditor/Private/CameraCalibrationCoreEditorModule.cpp
--- a/Engine/Plugins/VirtualProduction/CameraCalibrationCore/Source/CameraCalibrationCoreEditor/Private/CameraCalibrationCoreEditorModule.cpp
+++ b/Engine/Plugins/VirtualProduction/CameraCalibrationCore/Source/CameraCalibrationCoreEditor/Private/CameraCalibrationCoreEditorModule.cpp
@@ -8,6 +8,7 @@
 #include "ActorFactories/ActorFactory.h"
 #include "CalibrationPointComponent.h"
 #include "CalibrationPointComponentDetails.h"
+#include "CameraCalibrationCoreEditorStyle.h"
 #include "DistortionHandlerPickerDetailCustomization.h"
 #include "Editor.h"
 #include "IPlacementModeModule.h"
@@ -23,6 +24,8 @@ DEFINE_LOG_CATEGORY(LogCameraCalibrationCoreEditor);
 
 void FCameraCalibrationCoreEditorModule::StartupModule()
 {
+	FCameraCalibrationCoreEditorStyle::Get();
+
 	FPropertyEditorModule& PropertyEditorModule = FModuleManager::Get().LoadModuleChecked<FPropertyEditorModule>("PropertyEditor");
 
 	PropertyEditorModule.RegisterCustomPropertyTypeLayout(
@@ -82,6 +85,7 @@ const FPlacementCategoryInfo* FCameraCalibrationCoreEditorModule::GetVirtualProd
 	{
 		FPlacementCategoryInfo Info(
 			LOCTEXT("VirtualProductionCategoryName", "Virtual Production"),
+			FSlateIcon(FCameraCalibrationCoreEditorStyle::Get().GetStyleSetName(), "PlacementBrowser.Icons.VirtualProduction"),
 			VirtualProductionName,
 			TEXT("PMVirtualProduction"),
 			25 // Determines where the category shows up in the list with respect to the others.
diff --git a/Engine/Plugins/VirtualProduction/CameraCalibrationCore/Source/CameraCalibrationCoreEditor/Private/CameraCalibrationCoreEditorStyle.h b/Engine/Plugins/VirtualProduction/CameraCalibrationCore/Source/CameraCalibrationCoreEditor/Private/CameraCalibrationCoreEditorStyle.h
new file mode 100644
--- /dev/null
+++ b/Engine/Plugins/VirtualProduction/CameraCalibrationCore/Source/CameraCalibrationCoreEditor/Private/CameraCalibrationCoreEditorStyle.h
@@ -0,0 +1,41 @@
+// Copyright Epic Games, Inc. All Rights Reserved.
+
+#pragma once
+
+#include "CoreMinimal.h"
+#include "Styling/SlateStyleRegistry.h"
+#include "Styling/SlateStyleMacros.h"
+#include "Styling/SlateStyle.h"
+
+
+/**
+ * Implements the visual style of the camera calibration core  UI.
+ */
+class FCameraCalibrationCoreEditorStyle : public FSlateStyleSet
+{
+public:
+
+	/** Default constructor. */
+	FCameraCalibrationCoreEditorStyle() : FSlateStyleSet("CameraCalibrationCoreEditorStyle")
+	{
+		const FVector2D Icon16x16(16.0f, 16.0f);
+
+		// Set placement browser icons
+		SetContentRoot(FPaths::EngineContentDir() / TEXT("Editor/Slate/"));
+		Set("PlacementBrowser.Icons.VirtualProduction", new IMAGE_BRUSH_SVG("Starship/Common/VirtualProduction", Icon16x16));
+
+		FSlateStyleRegistry::RegisterSlateStyle(*this);
+	}
+
+	/** Virtual destructor. */
+	virtual ~FCameraCalibrationCoreEditorStyle()
+	{
+		FSlateStyleRegistry::UnRegisterSlateStyle(*this);
+	}
+
+	static FCameraCalibrationCoreEditorStyle& Get()
+	{
+		static FCameraCalibrationCoreEditorStyle Inst;
+		return Inst;
+	}
+};
