LensFile Exporter

Add an exporter for the LensFile asset that outputs a .ulens file
The exporter doesn't handle STMaps yet, but should be fairly straight forward to add support for it.

diff --git a/Engine/Plugins/VirtualProduction/CameraCalibration/Source/CameraCalibrationEditor/CameraCalibrationEditor.Build.cs b/Engine/Plugins/VirtualProduction/CameraCalibration/Source/CameraCalibrationEditor/CameraCalibrationEditor.Build.cs
--- a/Engine/Plugins/VirtualProduction/CameraCalibration/Source/CameraCalibrationEditor/CameraCalibrationEditor.Build.cs
+++ b/Engine/Plugins/VirtualProduction/CameraCalibration/Source/CameraCalibrationEditor/CameraCalibrationEditor.Build.cs
@@ -24,6 +24,8 @@ namespace UnrealBuildTool.Rules
 					"EditorWidgets",
 					"Engine",
 					"InputCore",
+					"Json",
+					"JsonUtilities",
 					"LiveLinkCamera",
 					"LiveLinkComponents",
 					"LiveLinkEditor",
diff --git a/Engine/Plugins/VirtualProduction/CameraCalibration/Source/CameraCalibrationEditor/Private/Exporters/LensFileExporter.cpp b/Engine/Plugins/VirtualProduction/CameraCalibration/Source/CameraCalibrationEditor/Private/Exporters/LensFileExporter.cpp
new file mode 100644
--- /dev/null
+++ b/Engine/Plugins/VirtualProduction/CameraCalibration/Source/CameraCalibrationEditor/Private/Exporters/LensFileExporter.cpp
@@ -0,0 +1,404 @@
+// Copyright Epic Games, Inc. All Rights Reserved.
+
+#include "LensFileExporter.h"
+#include "JsonObjectConverter.h"
+#include "CameraCalibrationEditorLog.h"
+#include "CameraCalibrationUtils.h"
+#include "LensFile.h"
+#include "LensData.h"
+#include "Models/LensModel.h"
+
+
+FLensInfoExchange::FLensInfoExchange(const ULensFile* LensFile)
+{
+	if (LensFile == nullptr)
+	{
+		return;
+	}
+
+	const FLensInfo& LensInfo = LensFile->LensInfo;
+
+	SerialNumber = *LensInfo.LensSerialNumber;
+	ModelName = *LensInfo.LensModelName;
+
+	if (const ULensModel* LensModelObject = LensInfo.LensModel->GetDefaultObject<ULensModel>())
+	{
+		DistortionModel = LensModelObject->GetShortModelName();
+	}
+}
+
+FLensFileMetadata::FLensFileMetadata(const ULensFile* LensFile)
+	: LensInfo{ LensFile }
+{
+	Type = TEXT("LensFile");
+	Version = FLensFileExchange::LensFileVersion;
+
+	if (LensFile == nullptr)
+	{
+		return;
+	}
+
+	// Set the default values for metadata when exporting from the ULensFile
+	Name = *LensFile->GetName();
+
+	// Add the UserMetadata
+	UserMetadata.Reserve(LensFile->UserMetadata.Num());
+	for (const TPair<FString, FString>& UserMetadataPair : LensFile->UserMetadata)
+	{
+		FLensFileUserMetadataEntry UserMetadataEntry;
+		UserMetadataEntry.Name = *UserMetadataPair.Key;
+		UserMetadataEntry.Value = *UserMetadataPair.Value;
+		UserMetadata.Add(UserMetadataEntry);
+	}
+}
+
+FLensFileExchange::FLensFileExchange(const ULensFile* LensFile)
+	: Metadata{ LensFile }
+{
+	if (LensFile == nullptr)
+	{
+		return;
+	}
+
+	const FLensInfo& LensInfo = LensFile->LensInfo;
+	const FVector2D& Sensor = LensInfo.SensorDimensions;
+
+	SensorDimensions.Width = Sensor.X;
+	SensorDimensions.Height = Sensor.Y;
+	SensorDimensions.Units = ELensFileUnit::Millimeters;
+
+	ImageDimensions.Width = LensInfo.ImageDimensions.X;
+	ImageDimensions.Height = LensInfo.ImageDimensions.Y;
+
+	ExtractFocalLengthTable(LensFile);
+	ExtractImageCenterTable(LensFile);
+	ExtractNodalOffsetTable(LensFile);
+	ExtractEncoderTables(LensFile);
+	ExtractDistortionParameters(LensFile);
+	ExtractSTMaps(LensFile);
+}
+
+void FLensFileExchange::ExtractFocalLengthTable(const ULensFile* LensFile)
+{
+	// Read the FocalLength table
+	FLensFileParameterTable FocalLengthParametersTable;
+	FocalLengthParametersTable.ParameterName = GET_MEMBER_NAME_CHECKED(ULensFile, FocalLengthTable);
+	FocalLengthParametersTable.Header = { FocusEncoderHeaderName, ZoomEncoderHeaderName, FocalLengthFxHeaderName, FocalLengthFyHeaderName };
+
+	const FFocalLengthTable& LensFileFocalLengthTable = LensFile->FocalLengthTable;
+
+	for (int32 FocusPointIndex = 0; FocusPointIndex < LensFileFocalLengthTable.GetFocusPointNum(); ++FocusPointIndex)
+	{
+		const FFocalLengthFocusPoint& FocalLengthFocusPoint = LensFileFocalLengthTable.FocusPoints[FocusPointIndex];
+		const float Focus = FocalLengthFocusPoint.GetFocus();
+
+		for (int32 FocalLengthInfoIndex = 0; FocalLengthInfoIndex < FocalLengthFocusPoint.GetNumPoints(); ++FocalLengthInfoIndex)
+		{
+			FFocalLengthInfo FocalLengthInfo;
+			if (FocalLengthFocusPoint.GetValue(FocalLengthInfoIndex, FocalLengthInfo))
+			{
+				const float Zoom = FocalLengthFocusPoint.GetZoom(FocalLengthInfoIndex);
+				const float Fx = FocalLengthInfo.FxFy.X;
+				const float Fy = FocalLengthInfo.FxFy.Y;
+
+				FocalLengthParametersTable.Data.Add(Focus);
+				FocalLengthParametersTable.Data.Add(Zoom);
+				FocalLengthParametersTable.Data.Add(Fx);
+				FocalLengthParametersTable.Data.Add(Fy);
+			}
+		}
+	}
+
+	CameraParameterTables.Add(FocalLengthParametersTable);
+}
+
+void FLensFileExchange::ExtractImageCenterTable(const ULensFile* LensFile)
+{
+	// Read the Image Center parameter table
+	FLensFileParameterTable ImageCenterParametersTable;
+	ImageCenterParametersTable.ParameterName = GET_MEMBER_NAME_CHECKED(ULensFile, ImageCenterTable);
+	ImageCenterParametersTable.Header = { FocusEncoderHeaderName, ZoomEncoderHeaderName, ImageCenterCxHeaderName, ImageCenterCyHeaderName };
+
+	const FImageCenterTable& LensFileImageCenterTable = LensFile->ImageCenterTable;
+
+	for (int32 ImageCenterFocusPointIndex = 0; ImageCenterFocusPointIndex < LensFileImageCenterTable.GetFocusPointNum(); ++ImageCenterFocusPointIndex)
+	{
+		const FImageCenterFocusPoint& ImageCenterFocusPoint = LensFileImageCenterTable.FocusPoints[ImageCenterFocusPointIndex];
+		const float Focus = ImageCenterFocusPoint.GetFocus();
+
+		for (int32 ImageCenterInfoIndex = 0; ImageCenterInfoIndex < ImageCenterFocusPoint.GetNumPoints(); ++ImageCenterInfoIndex)
+		{
+			const float Zoom = ImageCenterFocusPoint.GetZoom(ImageCenterInfoIndex);
+
+			FImageCenterInfo ImageCenterInfo;
+			if (ImageCenterFocusPoint.GetPoint(Zoom, ImageCenterInfo))
+			{
+				const float Cx = ImageCenterInfo.PrincipalPoint.X;
+				const float Cy = ImageCenterInfo.PrincipalPoint.Y;
+
+				ImageCenterParametersTable.Data.Add(Focus);
+				ImageCenterParametersTable.Data.Add(Zoom);
+				ImageCenterParametersTable.Data.Add(Cx);
+				ImageCenterParametersTable.Data.Add(Cy);
+			}
+		}
+	}
+
+	CameraParameterTables.Add(ImageCenterParametersTable);
+}
+
+void FLensFileExchange::ExtractNodalOffsetTable(const ULensFile* LensFile)
+{
+	// Read the NodalOffset parameter table
+	FLensFileParameterTable NodalOffsetParametersTable;
+	NodalOffsetParametersTable.ParameterName = GET_MEMBER_NAME_CHECKED(ULensFile, NodalOffsetTable);
+	NodalOffsetParametersTable.Header = {
+		FocusEncoderHeaderName,
+		ZoomEncoderHeaderName,
+		NodalOffsetQxHeaderName,
+		NodalOffsetQyHeaderName,
+		NodalOffsetQzHeaderName,
+		NodalOffsetQwHeaderName,
+		NodalOffsetTxHeaderName,
+		NodalOffsetTyHeaderName,
+		NodalOffsetTzHeaderName
+	};
+
+	const FNodalOffsetTable& LensFileNodalOffsetTable = LensFile->NodalOffsetTable;
+
+	for (int32 NodalOffsetFocusPointIndex = 0; NodalOffsetFocusPointIndex < LensFileNodalOffsetTable.GetFocusPointNum(); ++NodalOffsetFocusPointIndex)
+	{
+		const FNodalOffsetFocusPoint& NodalOffsetFocusPoint = LensFileNodalOffsetTable.FocusPoints[NodalOffsetFocusPointIndex];
+		const float Focus = NodalOffsetFocusPoint.GetFocus();
+
+		for (int32 NodalOffsetPointIndex = 0; NodalOffsetPointIndex < NodalOffsetFocusPoint.GetNumPoints(); ++NodalOffsetPointIndex)
+		{
+			const float Zoom = NodalOffsetFocusPoint.GetZoom(NodalOffsetPointIndex);
+
+			FNodalPointOffset NodalOffsetInfo;
+			if (NodalOffsetFocusPoint.GetPoint(Zoom, NodalOffsetInfo))
+			{
+				NodalOffsetParametersTable.Data.Add(Focus);
+				NodalOffsetParametersTable.Data.Add(Zoom);
+
+				FTransform NodalOffsetTransform{ NodalOffsetInfo.RotationOffset, NodalOffsetInfo.LocationOffset };
+				FCameraCalibrationUtils::ConvertUnrealToOpenCV(NodalOffsetTransform);
+
+				const FQuat RotationOffset = NodalOffsetTransform.GetRotation();
+				const FVector LocationOffset = NodalOffsetTransform.GetTranslation();
+
+				NodalOffsetParametersTable.Data.Add(RotationOffset.X);
+				NodalOffsetParametersTable.Data.Add(RotationOffset.Y);
+				NodalOffsetParametersTable.Data.Add(RotationOffset.Z);
+				NodalOffsetParametersTable.Data.Add(RotationOffset.W);
+				NodalOffsetParametersTable.Data.Add(LocationOffset.X);
+				NodalOffsetParametersTable.Data.Add(LocationOffset.Y);
+				NodalOffsetParametersTable.Data.Add(LocationOffset.Z);
+			}
+		}
+
+		CameraParameterTables.Add(NodalOffsetParametersTable);
+	}
+}
+
+void FLensFileExchange::ExtractEncoderTables(const ULensFile* LensFile)
+{
+	// Read the Encoder Tables from the LensFile
+	FLensFileParameterTable FocusEncordersTable;
+	FLensFileParameterTable IrisEncodersTable;
+
+	FocusEncordersTable.ParameterName = GET_MEMBER_NAME_CHECKED(FEncodersTable, Focus);
+	FocusEncordersTable.Header = { FocusEncoderHeaderName, FocusCMHeaderName };
+
+	IrisEncodersTable.ParameterName = GET_MEMBER_NAME_CHECKED(FEncodersTable, Iris);
+	IrisEncodersTable.Header = { IrisEncoderHeaderName, IrisFstopHeaderName };
+
+	const FEncodersTable& LensFileEncodersTable = LensFile->EncodersTable;
+
+	FocusEncordersTable.Data.Reserve(LensFileEncodersTable.GetNumFocusPoints() * 2);
+	IrisEncodersTable.Data.Reserve(LensFileEncodersTable.GetNumIrisPoints() * 2);
+
+	for (int32 FocusPointIndex = 0; FocusPointIndex < LensFileEncodersTable.GetNumFocusPoints(); ++FocusPointIndex)
+	{
+		const float FocusInput = LensFileEncodersTable.GetFocusInput(FocusPointIndex);
+		const float FocusValue = LensFileEncodersTable.GetFocusValue(FocusPointIndex);
+
+		FocusEncordersTable.Data.Add(FocusInput);
+		FocusEncordersTable.Data.Add(FocusValue);
+	}
+
+	for (int32 IrisPointIndex = 0; IrisPointIndex < LensFileEncodersTable.GetNumIrisPoints(); ++IrisPointIndex)
+	{
+		const float IrisInput = LensFileEncodersTable.GetIrisInput(IrisPointIndex);
+		const float IrisValue = LensFileEncodersTable.GetIrisValue(IrisPointIndex);
+
+		IrisEncodersTable.Data.Add(IrisInput);
+		IrisEncodersTable.Data.Add(IrisValue);
+	}
+
+	EncoderTables.Add(FocusEncordersTable);
+	EncoderTables.Add(IrisEncodersTable);
+}
+
+void FLensFileExchange::ExtractDistortionParameters(const ULensFile* LensFile)
+{
+	// Read the Distortion parameter table
+	FLensFileParameterTable DistortionParametersTable;
+	DistortionParametersTable.ParameterName = GET_MEMBER_NAME_CHECKED(ULensFile, DistortionTable);
+
+	if (const ULensModel* LensModelObject = LensFile->LensInfo.LensModel->GetDefaultObject<ULensModel>())
+	{
+		const int32 NumDistortionParams = LensModelObject->GetNumParameters();
+		DistortionParametersTable.Header.Reserve(NumDistortionParams + 2);
+
+		// Add the Focus and Zoom encoders headers that are common to all camera parameters tables
+		DistortionParametersTable.Header.Add(FocusEncoderHeaderName);
+		DistortionParametersTable.Header.Add(ZoomEncoderHeaderName);
+
+		for (const FText& ParamName : LensModelObject->GetParameterDisplayNames())
+		{
+			DistortionParametersTable.Header.Add(*ParamName.ToString());
+		}
+
+		const FDistortionTable& LensFileDistortionTable = LensFile->DistortionTable;
+		DistortionParametersTable.Data.Reserve(LensFileDistortionTable.GetFocusPointNum());
+
+		for (int32 DistortionFocusPointIndex = 0; DistortionFocusPointIndex < LensFileDistortionTable.GetFocusPointNum(); ++DistortionFocusPointIndex)
+		{
+			const FDistortionFocusPoint& DistortionFocusPoint = LensFileDistortionTable.FocusPoints[DistortionFocusPointIndex];
+			const float Focus = DistortionFocusPoint.GetFocus();
+
+			for (int32 DistortionPointIndex = 0; DistortionPointIndex < DistortionFocusPoint.GetNumPoints(); ++DistortionPointIndex)
+			{
+				const float Zoom = DistortionFocusPoint.GetZoom(DistortionPointIndex);
+
+				FDistortionInfo DistortionInfo;
+				if (DistortionFocusPoint.GetPoint(Zoom, DistortionInfo))
+				{
+					if (DistortionInfo.Parameters.Num() == NumDistortionParams)
+					{
+						// Copy the parameters to the exchange struct
+						DistortionParametersTable.Data.Add(Focus);
+						DistortionParametersTable.Data.Add(Zoom);
+
+						for (const float ParamValue : DistortionInfo.Parameters)
+						{
+							DistortionParametersTable.Data.Add(ParamValue);
+						}
+					}
+					else
+					{
+						UE_LOG(LogCameraCalibrationEditor, Warning, TEXT("Expected %d parameters for zoom %f but %d are available."), NumDistortionParams, Zoom, DistortionInfo.Parameters.Num());
+					}
+				}
+			}
+		}
+	}
+
+	CameraParameterTables.Add(DistortionParametersTable);
+}
+
+void FLensFileExchange::ExtractSTMaps(const ULensFile* LensFile)
+{
+	if (LensFile->STMapTable.GetFocusPointNum() > 0)
+	{
+		// TODO: Extract the STMaps from the LensFile to create a new FLensFileParameterTable
+		UE_LOG(LogCameraCalibrationEditor, Warning, TEXT("STMaps are not yet supported when exporting a LensFile"));
+	}
+}
+
+ULensFileExporter::ULensFileExporter(const FObjectInitializer& ObjectInitializer)
+	: Super(ObjectInitializer)
+{
+	SupportedClass = ULensFile::StaticClass();
+	bText = true;
+	PreferredFormatIndex = 0;
+	FormatExtension.Add(TEXT("ulens"));
+	FormatDescription.Add(TEXT("Unreal LensFile"));
+}
+
+bool ULensFileExporter::ExportText(const class FExportObjectInnerContext* Context, UObject* Object, const TCHAR* Type, FOutputDevice& Ar, FFeedbackContext* Warn, uint32 PortFlags)
+{
+	const ULensFile* LensFile = Cast<ULensFile>(Object);
+	if (LensFile == nullptr)
+	{
+		return false;
+	}
+
+	FLensFileExchange LensFileExchange{ LensFile };
+
+	FJsonObjectConverter::CustomExportCallback CustomExportCallback;
+	CustomExportCallback.BindLambda([](FProperty* Property, const void* Value) -> TSharedPtr<FJsonValue>
+	{
+		TSharedPtr<FJsonValue> Result = nullptr;
+
+		if (FStructProperty* PropertyAsStruct = CastField<FStructProperty>(Property))
+		{
+			// If this is a FLensFileParameterTable we need to do a custom export
+			if (PropertyAsStruct->Struct->IsChildOf(FLensFileParameterTable::StaticStruct()))
+			{
+				// Cast the property to our known struct type
+				const FLensFileParameterTable* ParameterTable = (const FLensFileParameterTable*)Value;
+
+				// Use the reflection system to get the FProperty representing the fields of the table
+				FProperty* ParameterNameProperty = ParameterTable->StaticStruct()->FindPropertyByName(GET_MEMBER_NAME_CHECKED(FLensFileParameterTable, ParameterName));
+				FProperty* HeaderProperty = ParameterTable->StaticStruct()->FindPropertyByName(GET_MEMBER_NAME_CHECKED(FLensFileParameterTable, Header));
+				FProperty* DataProperty = ParameterTable->StaticStruct()->FindPropertyByName(GET_MEMBER_NAME_CHECKED(FLensFileParameterTable, Data));
+
+				// Convert the Header and Data field to Strings
+				FString HeaderValuesStr;
+				const int32 NumHeaders = ParameterTable->Header.Num();
+				for (int32 HeaderIndex = 0; HeaderIndex < NumHeaders; ++HeaderIndex)
+				{
+					const FName& Header = ParameterTable->Header[HeaderIndex];
+					HeaderValuesStr.Append(Header.ToString());
+					if (HeaderIndex < (NumHeaders - 1))
+					{
+						HeaderValuesStr.Append(TEXT(", "));
+					}
+				}
+
+				FString DataStr;
+				const int32 NumDataElements = ParameterTable->Data.Num();
+				for (int32 DataIndex = 0; DataIndex < NumDataElements; ++DataIndex)
+				{
+					const float DataElement = ParameterTable->Data[DataIndex];
+					DataStr.Append(FString::SanitizeFloat(DataElement));
+					if (DataIndex < (NumDataElements - 1))
+					{
+						DataStr.Append(((DataIndex + 1) % NumHeaders == 0) ? TEXT("; ") : TEXT(", "));
+					}
+				}
+
+				// Convert the ParameterNameProperty to a Json representation
+				TSharedPtr<FJsonValue> ParameterNameJson = FJsonObjectConverter::UPropertyToJsonValue(ParameterNameProperty, &ParameterTable->ParameterName);
+
+				// Json Object that will hold the serialized FLensFileParameterTable
+				TSharedPtr<FJsonObject> ParameterTableJson = MakeShared<FJsonObject>();
+				ParameterTableJson->SetField(FJsonObjectConverter::StandardizeCase(ParameterNameProperty->GetName()), ParameterNameJson);
+				ParameterTableJson->SetStringField(FJsonObjectConverter::StandardizeCase(HeaderProperty->GetName()), HeaderValuesStr);
+				ParameterTableJson->SetStringField(FJsonObjectConverter::StandardizeCase(DataProperty->GetName()), DataStr);
+
+				Result = MakeShared<FJsonValueObject>(ParameterTableJson);
+			}
+		}
+
+		return Result;
+	});
+
+	const int64 CheckFlags = 0;
+	const int64 SkipFlags = 0;
+	const int32 Indent = 0;
+	const bool bPrettyPrint = true;
+	FString SerializedJson;
+	if (FJsonObjectConverter::UStructToJsonObjectString<FLensFileExchange>(LensFileExchange, SerializedJson, CheckFlags, SkipFlags, Indent, &CustomExportCallback, bPrettyPrint))
+	{
+		Ar.Log(SerializedJson);
+
+		return true;
+	}
+	else
+	{
+		return false;
+	}
+}
\ No newline at end of file
diff --git a/Engine/Plugins/VirtualProduction/CameraCalibration/Source/CameraCalibrationEditor/Private/Exporters/LensFileExporter.h b/Engine/Plugins/VirtualProduction/CameraCalibration/Source/CameraCalibrationEditor/Private/Exporters/LensFileExporter.h
new file mode 100644
--- /dev/null
+++ b/Engine/Plugins/VirtualProduction/CameraCalibration/Source/CameraCalibrationEditor/Private/Exporters/LensFileExporter.h
@@ -0,0 +1,193 @@
+// Copyright Epic Games, Inc. All Rights Reserved.
+
+#pragma once
+
+#include "CoreMinimal.h"
+#include "UObject/ObjectMacros.h"
+#include "Exporters/Exporter.h"
+#include "LensFileExporter.generated.h"
+
+
+// TODO: Move the Structs to another file
+
+UENUM()
+enum class ELensFileUnit
+{
+	Millimeters,
+	Normalized
+};
+
+UENUM()
+enum class ENodalOffsetCoordinateSystem
+{
+	OpenCV
+};
+
+USTRUCT()
+struct FLensInfoExchange
+{
+	GENERATED_BODY()
+
+	FLensInfoExchange(const class ULensFile* LensFile = nullptr);
+
+	UPROPERTY()
+	FName SerialNumber;
+
+	UPROPERTY()
+	FName ModelName;
+
+	UPROPERTY()
+	FName DistortionModel;
+};
+
+USTRUCT()
+struct FLensFileUserMetadataEntry
+{
+	GENERATED_BODY()
+
+	UPROPERTY()
+	FName Name;
+
+	UPROPERTY()
+	FName Value;
+};
+
+USTRUCT()
+struct FLensFileMetadata
+{
+	GENERATED_BODY()
+
+	FLensFileMetadata(const class ULensFile* LensFile = nullptr);
+
+	UPROPERTY()
+	FName Type;
+
+	UPROPERTY()
+	FName Version;
+
+	UPROPERTY()
+	FLensInfoExchange LensInfo;
+
+	UPROPERTY()
+	FName Name;
+
+	UPROPERTY()
+	ENodalOffsetCoordinateSystem NodalOffsetCoordinateSystem = ENodalOffsetCoordinateSystem::OpenCV;
+
+	UPROPERTY()
+	ELensFileUnit FxFyUnits = ELensFileUnit::Normalized;
+
+	UPROPERTY()
+	ELensFileUnit CxCyUnits = ELensFileUnit::Normalized;
+
+	UPROPERTY()
+	TArray<FLensFileUserMetadataEntry> UserMetadata;
+};
+
+
+USTRUCT()
+struct FLensFileSensorDimensions
+{
+	GENERATED_BODY()
+
+	UPROPERTY()
+	float Width = 0.0f;
+
+	UPROPERTY()
+	float Height = 0.0f;
+
+	UPROPERTY()
+	ELensFileUnit Units = ELensFileUnit::Normalized;
+};
+
+USTRUCT()
+struct FLensFileImageDimensions
+{
+	GENERATED_BODY()
+
+	UPROPERTY()
+	int32 Width = 0;
+
+	UPROPERTY()
+	int32 Height = 0;
+};
+
+USTRUCT()
+struct FLensFileParameterTable
+{
+	GENERATED_BODY()
+
+	UPROPERTY()
+	FName ParameterName;
+
+	UPROPERTY()
+	TArray<FName> Header;
+
+	UPROPERTY()
+	TArray<float> Data;
+};
+
+USTRUCT()
+struct FLensFileExchange
+{
+	GENERATED_BODY()
+
+	FLensFileExchange(const class ULensFile* LensFile = nullptr);
+
+	UPROPERTY()
+	FLensFileMetadata Metadata;
+
+	UPROPERTY()
+	FLensFileSensorDimensions SensorDimensions;
+
+	UPROPERTY()
+	FLensFileImageDimensions ImageDimensions;
+
+	UPROPERTY()
+	TArray<FLensFileParameterTable> CameraParameterTables;
+
+	UPROPERTY()
+	TArray<FLensFileParameterTable> EncoderTables;
+
+	constexpr static TCHAR LensFileVersion[] = TEXT("0.0.0");
+
+	constexpr static TCHAR FocusEncoderHeaderName[] = TEXT("FocusEncoder");
+	constexpr static TCHAR ZoomEncoderHeaderName[] = TEXT("ZoomEncoder");
+	constexpr static TCHAR IrisEncoderHeaderName[] = TEXT("IrisEncoder");
+
+	constexpr static TCHAR FocusCMHeaderName[] = TEXT("FocusCM");
+	constexpr static TCHAR IrisFstopHeaderName[] = TEXT("IrisFstop");
+
+	constexpr static TCHAR FocalLengthFxHeaderName[] = TEXT("Fx");
+	constexpr static TCHAR FocalLengthFyHeaderName[] = TEXT("Fy");
+
+	constexpr static TCHAR ImageCenterCxHeaderName[] = TEXT("Cx");
+	constexpr static TCHAR ImageCenterCyHeaderName[] = TEXT("Cy");
+
+	constexpr static TCHAR NodalOffsetQxHeaderName[] = TEXT("Qx");
+	constexpr static TCHAR NodalOffsetQyHeaderName[] = TEXT("Qy");
+	constexpr static TCHAR NodalOffsetQzHeaderName[] = TEXT("Qz");
+	constexpr static TCHAR NodalOffsetQwHeaderName[] = TEXT("Qw");
+
+	constexpr static TCHAR NodalOffsetTxHeaderName[] = TEXT("Tx");
+	constexpr static TCHAR NodalOffsetTyHeaderName[] = TEXT("Ty");
+	constexpr static TCHAR NodalOffsetTzHeaderName[] = TEXT("Tz");
+
+private:
+	void ExtractFocalLengthTable(const class ULensFile* LensFile);
+	void ExtractImageCenterTable(const class ULensFile* LensFile);
+	void ExtractNodalOffsetTable(const class ULensFile* LensFile);
+	void ExtractEncoderTables(const class ULensFile* LensFile);
+	void ExtractDistortionParameters(const class ULensFile* LensFile);
+	void ExtractSTMaps(const class ULensFile* LensFile);
+};
+
+UCLASS()
+class ULensFileExporter : public UExporter
+{
+	GENERATED_UCLASS_BODY()
+
+	//~ Begin UExporter interface
+	bool ExportText(const class FExportObjectInnerContext* Context, UObject* Object, const TCHAR* Type, FOutputDevice& Ar, FFeedbackContext* Warn, uint32 PortFlags/* =0 */) override;
+	//~ End UExporter interface
+};
\ No newline at end of file
diff --git a/Engine/Plugins/VirtualProduction/CameraCalibrationCore/Source/CameraCalibrationCore/Private/Models/SphericalLensModel.cpp b/Engine/Plugins/VirtualProduction/CameraCalibrationCore/Source/CameraCalibrationCore/Private/Models/SphericalLensModel.cpp
--- a/Engine/Plugins/VirtualProduction/CameraCalibrationCore/Source/CameraCalibrationCore/Private/Models/SphericalLensModel.cpp
+++ b/Engine/Plugins/VirtualProduction/CameraCalibrationCore/Source/CameraCalibrationCore/Private/Models/SphericalLensModel.cpp
@@ -10,4 +10,9 @@ UScriptStruct* USphericalLensModel::GetParameterStruct() const
 FName USphericalLensModel::GetModelName() const 
 { 
 	return FName("Spherical Lens Model"); 
+}
+
+FName USphericalLensModel::GetShortModelName() const
+{
+	return TEXT("Spherical");
 }
\ No newline at end of file
diff --git a/Engine/Plugins/VirtualProduction/CameraCalibrationCore/Source/CameraCalibrationCore/Private/Tables/FocalLengthTable.cpp b/Engine/Plugins/VirtualProduction/CameraCalibrationCore/Source/CameraCalibrationCore/Private/Tables/FocalLengthTable.cpp
--- a/Engine/Plugins/VirtualProduction/CameraCalibrationCore/Source/CameraCalibrationCore/Private/Tables/FocalLengthTable.cpp
+++ b/Engine/Plugins/VirtualProduction/CameraCalibrationCore/Source/CameraCalibrationCore/Private/Tables/FocalLengthTable.cpp
@@ -85,6 +85,8 @@ bool FFocalLengthFocusPoint::GetValue(int32 Index, FFocalLengthInfo& OutData) co
 	if(ZoomPoints.IsValidIndex(Index))
 	{
 		OutData = ZoomPoints[Index].FocalLengthInfo;
+
+		return true;
 	}
 
 	return false;
diff --git a/Engine/Plugins/VirtualProduction/CameraCalibrationCore/Source/CameraCalibrationCore/Public/Models/LensModel.h b/Engine/Plugins/VirtualProduction/CameraCalibrationCore/Source/CameraCalibrationCore/Public/Models/LensModel.h
--- a/Engine/Plugins/VirtualProduction/CameraCalibrationCore/Source/CameraCalibrationCore/Public/Models/LensModel.h
+++ b/Engine/Plugins/VirtualProduction/CameraCalibrationCore/Source/CameraCalibrationCore/Public/Models/LensModel.h
@@ -24,6 +24,9 @@ public:
 	/** Get the lens model name */
 	virtual FName GetModelName() const PURE_VIRTUAL(ULensModel::GetModelName, return FName(""););
 
+	/** Get the lens model short name */
+	virtual FName GetShortModelName() const PURE_VIRTUAL(ULensModel::GetShortModelName, return FName(""););
+
 	/** Get the struct of distortion parameters supported by this model */
 	virtual UScriptStruct* GetParameterStruct() const PURE_VIRTUAL(ULensModel::GetParameterStruct, return nullptr;);
 
diff --git a/Engine/Plugins/VirtualProduction/CameraCalibrationCore/Source/CameraCalibrationCore/Public/Models/SphericalLensModel.h b/Engine/Plugins/VirtualProduction/CameraCalibrationCore/Source/CameraCalibrationCore/Public/Models/SphericalLensModel.h
--- a/Engine/Plugins/VirtualProduction/CameraCalibrationCore/Source/CameraCalibrationCore/Public/Models/SphericalLensModel.h
+++ b/Engine/Plugins/VirtualProduction/CameraCalibrationCore/Source/CameraCalibrationCore/Public/Models/SphericalLensModel.h
@@ -51,5 +51,6 @@ public:
 	//~ Begin ULensModel interface
 	virtual UScriptStruct* GetParameterStruct() const override;
 	virtual FName GetModelName() const override;
+	virtual FName GetShortModelName() const override;
 	//~ End ULensModel interface
 };
