Support for Zoom over a 'point' in LensFile viewport

diff --git a/Engine/Plugins/VirtualProduction/CameraCalibration/Source/CameraCalibrationEditor/Private/AssetEditor/SSimulcamViewport.cpp b/Engine/Plugins/VirtualProduction/CameraCalibration/Source/CameraCalibrationEditor/Private/AssetEditor/SSimulcamViewport.cpp
--- a/Engine/Plugins/VirtualProduction/CameraCalibration/Source/CameraCalibrationEditor/Private/AssetEditor/SSimulcamViewport.cpp
+++ b/Engine/Plugins/VirtualProduction/CameraCalibration/Source/CameraCalibrationEditor/Private/AssetEditor/SSimulcamViewport.cpp
@@ -1,102 +1,31 @@
 // Copyright Epic Games, Inc. All Rights Reserved.
 
 #include "SSimulcamViewport.h"
+#include "Engine/Texture.h"
+#include "Viewport/SSimulcamEditorViewport.h"
 
-#include "EditorStyleSet.h"
-#include "Materials/Material.h"
-#include "NodalOffsetTool.h"
-#include "Styling/SlateBrush.h"
-#include "Widgets/Layout/SScaleBox.h"
-#include "Widgets/Images/SImage.h"
-#include "Widgets/Input/SButton.h"
-
-
-SSimulcamViewport::SSimulcamViewport()
-	: Collector(this)
-	, Material(nullptr)
-	, MaterialBrush(nullptr)
-	, TextureSampler(nullptr)
-{ }
 
 void SSimulcamViewport::Construct(const FArguments& InArgs, UTexture* InTexture)
 {
 	OnSimulcamViewportClicked = InArgs._OnSimulcamViewportClicked;
 
-	// The Slate brush that renders the material.
-	BrushImageSize = InArgs._BrushImageSize;
-
-	if (InTexture != nullptr)
-	{
-		// create wrapper material
-		Material = NewObject<UMaterial>(GetTransientPackage(), NAME_None, RF_Transient);
-
-		if (Material != nullptr)
-		{
-			TextureSampler = NewObject<UMaterialExpressionTextureSample>(Material);
-			{
-				TextureSampler->Texture = InTexture;
-				TextureSampler->AutoSetSampleType();
-			}
-
-			FExpressionOutput& Output = TextureSampler->GetOutputs()[0];
-			FExpressionInput& Input = Material->EmissiveColor;
-			{
-				Input.Expression = TextureSampler;
-				Input.Mask = Output.Mask;
-				Input.MaskR = Output.MaskR;
-				Input.MaskG = Output.MaskG;
-				Input.MaskB = Output.MaskB;
-				Input.MaskA = Output.MaskA;
-			}
-
-			FExpressionInput& Opacity = Material->Opacity;
-			{
-				Opacity.Expression = TextureSampler;
-				Opacity.Mask = Output.Mask;
-				Opacity.MaskR = 0;
-				Opacity.MaskG = 0;
-				Opacity.MaskB = 0;
-				Opacity.MaskA = 1;
-			}
+	Texture = TStrongObjectPtr<UTexture>(InTexture);
 
-			Material->BlendMode = BLEND_Opaque;
-
-			Material->Expressions.Add(TextureSampler);
-			Material->MaterialDomain = EMaterialDomain::MD_UI;
-			Material->PostEditChange();
-		}
-
-		// create Slate brush
-		MaterialBrush = MakeShared<FSlateBrush>();
-		{
-			MaterialBrush->SetResourceObject(Material);
-		}
-	}
+	TextureViewport = SNew(SSimulcamEditorViewport, SharedThis(this), InArgs._WithZoom.Get(), InArgs._WithPan.Get());
 
 	ChildSlot
 	[
-		SNew(SScaleBox)
-		.Stretch_Lambda([]() -> EStretch::Type { return EStretch::Fill;	})
-		[
-			SNew(SImage)
-			.Image(MaterialBrush.IsValid() ? MaterialBrush.Get() : FEditorStyle::GetBrush("WhiteTexture"))
-			.OnMouseButtonDown_Lambda([&](const FGeometry& MyGeometry, const FPointerEvent& MouseEvent)
-			{
-				OnSimulcamViewportClicked.ExecuteIfBound(MyGeometry, MouseEvent);
-				return FReply::Handled();
-			})
-		]
+		TextureViewport.ToSharedRef()
 	];
 }
 
-void SSimulcamViewport::Tick(const FGeometry& InAllottedGeometry, const double InCurrentTime, const float InDeltaTime)
+UTexture* SSimulcamViewport::GetTexture() const
 {
-	SCompoundWidget::Tick(InAllottedGeometry, InCurrentTime, InDeltaTime);
-
-	if (BrushImageSize.IsSet())
-	{
-		const FVector2D Size = BrushImageSize.Get();
-		MaterialBrush->ImageSize.X = Size.X;
-		MaterialBrush->ImageSize.Y = Size.Y;
-	}
+	return Texture.Get();
 }
+
+bool SSimulcamViewport::HasValidTextureResource() const
+{
+	UTexture* CurrentTexture = GetTexture();
+	return CurrentTexture != nullptr && CurrentTexture->GetResource() != nullptr;
+}
\ No newline at end of file
diff --git a/Engine/Plugins/VirtualProduction/CameraCalibration/Source/CameraCalibrationEditor/Private/AssetEditor/SSimulcamViewport.h b/Engine/Plugins/VirtualProduction/CameraCalibration/Source/CameraCalibrationEditor/Private/AssetEditor/SSimulcamViewport.h
--- a/Engine/Plugins/VirtualProduction/CameraCalibration/Source/CameraCalibrationEditor/Private/AssetEditor/SSimulcamViewport.h
+++ b/Engine/Plugins/VirtualProduction/CameraCalibration/Source/CameraCalibrationEditor/Private/AssetEditor/SSimulcamViewport.h
@@ -5,17 +5,17 @@
 #include "Widgets/DeclarativeSyntaxSupport.h"
 #include "Widgets/SCompoundWidget.h"
 #include "UObject/GCObject.h"
+#include "UObject/StrongObjectPtr.h"
 
-
-class UMaterial;
-class UMaterialExpressionTextureSample;
 class UTexture;
 
 struct FGeometry;
 struct FPointerEvent;
 
+class SSimulcamEditorViewport;
+
 /** Delegate to be executed when the viewport area is clicked */
-DECLARE_DELEGATE_TwoParams(FSimulcamViewportClickedEventHandler, const FGeometry& MyGeometry, const FPointerEvent& MouseEvent);
+DECLARE_DELEGATE_TwoParams(FSimulcamViewportClickedEventHandler, const FGeometry& MyGeometry, const FPointerEvent& PointerEvent);
 
 /**
  * UI to display the provided UTexture and respond to mouse clicks. 
@@ -23,18 +23,14 @@ DECLARE_DELEGATE_TwoParams(FSimulcamViewportClickedEventHandler, const FGeometry
 class SSimulcamViewport : public SCompoundWidget
 {
 public:
-	SLATE_BEGIN_ARGS(SSimulcamViewport) { }
-	SLATE_ATTRIBUTE(FVector2D, BrushImageSize);
+	SLATE_BEGIN_ARGS(SSimulcamViewport)
+	: _WithZoom(true), _WithPan(true)
+	{}
 	SLATE_EVENT(FSimulcamViewportClickedEventHandler, OnSimulcamViewportClicked)
+	SLATE_ATTRIBUTE(bool, WithZoom)
+	SLATE_ATTRIBUTE(bool, WithPan)
 	SLATE_END_ARGS()
 
-public:
-
-	/** Default constructor. */
-	SSimulcamViewport();
-
-public:
-
 	/**
 	 * Construct this widget
 	 *
@@ -43,56 +39,22 @@ public:
 	 */
 	void Construct(const FArguments& InArgs, UTexture* InTexture);
 
-	/**
-	 * Tick this widget
-	 * 
-	 * @param InAllottedGeometry Geometry of the widget.
-	 * @param InCurrentTime CurrentTime of the engine.
-	 * @param InDeltaTime Deltatime since last frame.
-	 */
-	virtual void Tick(const FGeometry& InAllottedGeometry, const double InCurrentTime, const float InDeltaTime) override;
-
-private:
-
-	class FInternalReferenceCollector : public FGCObject
-	{
-	public:
-		FInternalReferenceCollector(SSimulcamViewport* InObject)
-			: Object(InObject)
-		{
-		}
-
-		//~ FGCObject interface
-		virtual void AddReferencedObjects(FReferenceCollector& InCollector) override
-		{
-			InCollector.AddReferencedObject(Object->Material);
-			InCollector.AddReferencedObject(Object->TextureSampler);
-		}
-
-	private:
-		SSimulcamViewport* Object;
-	};
-
-	friend FInternalReferenceCollector;
+	/** Returns the current camera texture */
+	UTexture* GetTexture() const;
 
-	/** Collector to keep the UObject alive. */
-	FInternalReferenceCollector Collector;
+	/** Is the camera texture valid ? */
+	bool HasValidTextureResource() const;
 
-	/** The material that wraps the video texture for display in an SImage. */
-	UMaterial* Material;
-
-	/** The Slate brush that renders the material. */
-	TSharedPtr<FSlateBrush> MaterialBrush;
-
-	/** The video texture sampler in the wrapper material. */
-	UMaterialExpressionTextureSample* TextureSampler;
-
-	/** Brush image size attribute. */
-	TAttribute<FVector2D> BrushImageSize;
+	/** called whenever the viewport is clicked */
+	void OnViewportClicked(const FGeometry& MyGeometry, const FPointerEvent& PointerEvent) { OnSimulcamViewportClicked.ExecuteIfBound(MyGeometry, PointerEvent); }
 
 private:
 
 	/** Delegate to be executed when the viewport area is clicked */
 	FSimulcamViewportClickedEventHandler OnSimulcamViewportClicked;
+
+	TSharedPtr<SSimulcamEditorViewport> TextureViewport;
+
+	TStrongObjectPtr<UTexture> Texture;
 };
 
diff --git a/Engine/Plugins/VirtualProduction/CameraCalibration/Source/CameraCalibrationEditor/Private/AssetEditor/Viewport/SSimulcamEditorViewport.cpp b/Engine/Plugins/VirtualProduction/CameraCalibration/Source/CameraCalibrationEditor/Private/AssetEditor/Viewport/SSimulcamEditorViewport.cpp
new file mode 100644
--- /dev/null
+++ b/Engine/Plugins/VirtualProduction/CameraCalibration/Source/CameraCalibrationEditor/Private/AssetEditor/Viewport/SSimulcamEditorViewport.cpp
@@ -0,0 +1,241 @@
+// Copyright Epic Games, Inc. All Rights Reserved.
+
+#include "AssetEditor/Viewport/SSimulcamEditorViewport.h"
+
+#include "AssetEditor/SSimulcamViewport.h"
+#include "AssetEditor/Viewport/SimulcamEditorViewportClient.h"
+#include "Engine/Texture.h"
+#include "Engine/TextureRenderTargetVolume.h"
+#include "Engine/VolumeTexture.h"
+#include "Framework/Application/SlateApplication.h"
+#include "Framework/MultiBox/MultiBoxBuilder.h"
+#include "Slate/SceneViewport.h"
+#include "Widgets/Input/SComboButton.h"
+#include "Widgets/Input/SNumericEntryBox.h"
+#include "Widgets/Input/SSlider.h"
+#include "Widgets/SToolTip.h"
+#include "Widgets/SViewport.h"
+#include "Widgets/Text/STextBlock.h"
+
+#define LOCTEXT_NAMESPACE "STextureEditorViewport"
+
+void SSimulcamEditorViewport::Construct(const FArguments& InArgs, const TSharedRef<SSimulcamViewport>& InSimulcamViewport, const bool bWithZoom, const bool bWithPan)
+{
+	bIsRenderingEnabled = true;
+	SimulcamViewportWeakPtr = InSimulcamViewport;
+
+	ChildSlot
+	[
+		SNew(SVerticalBox)
+		+ SVerticalBox::Slot()
+		.FillHeight(1.0f)
+		[
+			SNew(SHorizontalBox)
+			+ SHorizontalBox::Slot()
+			.FillWidth(1.0f)
+			[
+				SNew(SVerticalBox)
+				+ SVerticalBox::Slot()
+				.FillHeight(1)
+				[
+					SNew(SOverlay)
+					// viewport canvas
+					+ SOverlay::Slot()
+					[
+						SAssignNew(ViewportWidget, SViewport)
+						.EnableGammaCorrection(false)
+						.IsEnabled(FSlateApplication::Get().GetNormalExecutionAttribute())
+						.ShowEffectWhenDisabled(false)
+						.EnableBlending(true)
+						.ToolTip(SNew(SToolTip).Text(this, &SSimulcamEditorViewport::GetDisplayedResolution))
+					]
+
+					// tool bar
+					+ SOverlay::Slot()
+					.Padding(2.0f)
+					.VAlign(VAlign_Top)
+					[
+						SNew(SHorizontalBox)
+						+ SHorizontalBox::Slot()
+						.FillWidth(1.0f)
+						.Padding(4.0f, 0.0f)
+						.HAlign(HAlign_Right)
+						.VAlign(VAlign_Center)
+						[
+							SNew(STextBlock)
+							.Text(this, &SSimulcamEditorViewport::HandleZoomPercentageText)
+							.Visibility(bWithZoom ? EVisibility::Visible : EVisibility::Hidden)
+						]
+					]
+				]
+			]
+		]
+
+	];
+
+	ViewportClient = MakeShared<FSimulcamEditorViewportClient>(InSimulcamViewport, SharedThis(this), bWithZoom, bWithPan);
+	Viewport = MakeShared<FSceneViewport>(ViewportClient.Get(), ViewportWidget);
+	// The viewport widget needs an interface so it knows what should render
+	ViewportWidget->SetViewportInterface(Viewport.ToSharedRef());
+
+	// required for having the first 'zoom to fit' behaviour
+	CacheEffectiveTextureSize();
+
+	Viewport->ViewportResizedEvent.AddSP(ViewportClient.ToSharedRef(), &FSimulcamEditorViewportClient::OnViewportResized);
+}
+
+void SSimulcamEditorViewport::EnableRendering()
+{
+	bIsRenderingEnabled = true;
+}
+
+void SSimulcamEditorViewport::DisableRendering()
+{
+	bIsRenderingEnabled = false;
+}
+
+TSharedPtr<FSceneViewport> SSimulcamEditorViewport::GetViewport() const
+{
+	return Viewport;
+}
+
+TSharedPtr<SViewport> SSimulcamEditorViewport::GetViewportWidget() const
+{
+	return ViewportWidget;
+}
+
+void SSimulcamEditorViewport::Tick(const FGeometry& AllottedGeometry, const double InCurrentTime, const float InDeltaTime)
+{
+	if (bIsRenderingEnabled)
+	{
+		Viewport->Invalidate();
+	}
+}
+
+FText SSimulcamEditorViewport::GetDisplayedResolution() const
+{
+	return ViewportClient->GetDisplayedResolution();
+}
+
+bool SSimulcamEditorViewport::HasValidTextureResource() const
+{
+	return (SimulcamViewportWeakPtr.IsValid() && SimulcamViewportWeakPtr.Pin()->HasValidTextureResource());
+}
+
+FText SSimulcamEditorViewport::HandleZoomPercentageText() const
+{
+	if (SimulcamViewportWeakPtr.IsValid() && HasValidTextureResource())
+	{
+		FText FormattedText = FText::FromString(TEXT("{0}: {1}"));
+		double DisplayedZoomLevel = GetCustomZoomLevel();
+		return FText::Format(FormattedText, FText::FromName(SimulcamViewportWeakPtr.Pin()->GetTexture()->GetFName()), FText::AsPercent(DisplayedZoomLevel));
+	}
+	else
+	{
+		return LOCTEXT("InvalidTexture", "Invalid Texture");
+	}
+}
+
+void SSimulcamEditorViewport::CacheEffectiveTextureSize()
+{
+	if (!SimulcamViewportWeakPtr.IsValid())
+	{
+		return;
+	}
+
+	UTexture* Texture = SimulcamViewportWeakPtr.Pin()->GetTexture();
+	if (!Texture)
+	{
+		return;
+	}
+
+	CachedEffectiveTextureSize = FVector2D(Texture->GetSurfaceWidth(), Texture->GetSurfaceHeight());
+}
+
+FVector2D SSimulcamEditorViewport::CalculateTextureDimensions() const
+{
+	if (!SimulcamViewportWeakPtr.IsValid())
+	{
+		return FVector2D::ZeroVector;
+	}
+
+	UTexture* Texture = SimulcamViewportWeakPtr.Pin()->GetTexture();
+	if (!Texture)
+	{
+		return FVector2D::ZeroVector;
+	}
+
+	FVector2D Size = FVector2D(CachedEffectiveTextureSize.X * Zoom, CachedEffectiveTextureSize.Y * Zoom);
+
+	/** catch if the Width and Height are still zero for some reason */
+	if ((Size.X == 0) || (Size.Y == 0))
+	{
+		Size = FVector2D::ZeroVector;
+	}
+
+	return Size;
+}
+
+double SSimulcamEditorViewport::GetCustomZoomLevel() const
+{
+	return Zoom;
+}
+
+double SSimulcamEditorViewport::GetMinZoomLevel() const
+{
+	FIntPoint ViewportSize = GetViewport()->GetSizeXY();
+	return FMath::Min(ViewportSize.X / CachedEffectiveTextureSize.X, ViewportSize.Y / CachedEffectiveTextureSize.Y);
+}
+
+FVector2D SSimulcamEditorViewport::GetFitPosition() const
+{
+	FIntPoint ViewportSize = GetViewport()->GetSizeXY();
+	double MinZoom = GetMinZoomLevel();
+	return FVector2D((ViewportSize.X - (CachedEffectiveTextureSize.X * MinZoom)) / 2, (ViewportSize.Y - (CachedEffectiveTextureSize.Y * MinZoom)) / 2);
+}
+
+void SSimulcamEditorViewport::SetCustomZoomLevel(double ZoomValue)
+{
+	Zoom = FMath::Clamp(ZoomValue, GetMinZoomLevel(), MaxZoom);
+}
+
+void SSimulcamEditorViewport::OffsetZoom(double OffsetValue, bool bSnapToStepSize)
+{
+	/** Offset from our current "visual" zoom level so that you can
+		smoothly transition from Fit/Fill mode into a custom zoom level */
+	const double CurrentZoom = GetCustomZoomLevel();
+	if (bSnapToStepSize)
+	{
+		/** Snap to the zoom step when offsetting to avoid zooming all the way to the min(0.01)
+			then back up (+0.1) causing your zoom level to be off by 0.01 (eg. 11%)
+			If we were in a fit view mode then our current zoom level could also be off the grid */
+		const double FinalZoom = FMath::GridSnap(CurrentZoom + OffsetValue, ZoomStep);
+		SetCustomZoomLevel(FinalZoom);
+	}
+	else
+	{
+		SetCustomZoomLevel(CurrentZoom + OffsetValue);
+	}
+}
+
+void SSimulcamEditorViewport::ZoomIn()
+{
+	OffsetZoom(ZoomStep);
+}
+
+void SSimulcamEditorViewport::ZoomOut()
+{
+	OffsetZoom(-ZoomStep);
+}
+
+void SSimulcamEditorViewport::OnViewportClicked(const FGeometry& MyGeometry, const FPointerEvent& PointerEvent)
+{
+	if (!SimulcamViewportWeakPtr.IsValid())
+	{
+		return;
+	}
+
+	SimulcamViewportWeakPtr.Pin()->OnViewportClicked(MyGeometry, PointerEvent);
+}
+
+#undef LOCTEXT_NAMESPACE
diff --git a/Engine/Plugins/VirtualProduction/CameraCalibration/Source/CameraCalibrationEditor/Private/AssetEditor/Viewport/SSimulcamEditorViewport.h b/Engine/Plugins/VirtualProduction/CameraCalibration/Source/CameraCalibrationEditor/Private/AssetEditor/Viewport/SSimulcamEditorViewport.h
new file mode 100644
--- /dev/null
+++ b/Engine/Plugins/VirtualProduction/CameraCalibration/Source/CameraCalibrationEditor/Private/AssetEditor/Viewport/SSimulcamEditorViewport.h
@@ -0,0 +1,112 @@
+// Copyright Epic Games, Inc. All Rights Reserved.
+#pragma once
+#include "CoreMinimal.h"
+#include "Layout/Visibility.h"
+#include "Widgets/DeclarativeSyntaxSupport.h"
+#include "Widgets/SCompoundWidget.h"
+class FSceneViewport;
+class SViewport;
+class SSimulcamViewport;
+class FSimulcamEditorViewportClient;
+struct FPointerEvent;
+
+/**
+ * Implements the texture editor's view port.
+ */
+class SSimulcamEditorViewport
+	: public SCompoundWidget
+{
+public:
+	SLATE_BEGIN_ARGS(SSimulcamEditorViewport) { }
+	SLATE_END_ARGS()
+public:
+
+	/**
+	 * Constructs the widget.
+	 *
+	 * @param InArgs The construction arguments.
+	 */
+	void Construct( const FArguments& InArgs, const TSharedRef<SSimulcamViewport>& InSimulcamViewport, const bool bWithZoom, const bool bWithPan);
+
+	/** Enable viewport rendering */
+	void EnableRendering();
+	/** Disable viewport rendering */
+	void DisableRendering();
+
+	/** Get the current SceneViewport */
+	TSharedPtr<FSceneViewport> GetViewport( ) const;
+
+	/** Get The SViewport Widget using this editor */
+	TSharedPtr<SViewport> GetViewportWidget( ) const;
+
+	/** Returns the current texture dimensions (in pixels) */
+	FVector2D CalculateTextureDimensions() const;
+
+	/** The current zoom level */
+	double GetCustomZoomLevel() const;
+
+	/** Directly set a zoom level */
+	void SetCustomZoomLevel(double ZoomValue);
+
+	/** Zoom with offset */
+	void OffsetZoom(double OffsetValue, bool bSnapToStepSize = true);
+	
+	/** Increase zooming */
+	void ZoomIn();
+
+	/** Decrease zooming */
+	void ZoomOut();
+
+	//~ Begin SWidget interface
+	virtual void Tick( const FGeometry& AllottedGeometry, const double InCurrentTime, const float InDeltaTime ) override;
+	//~ End SWidget interface
+
+	/** Quick compute and cache the texture size before each frame */
+	void CacheEffectiveTextureSize();
+
+	/** Returns the minimal level of Zoom */
+	double GetMinZoomLevel() const;
+
+	/** Compute the best position for fitting the texture in the viewport */
+	FVector2D GetFitPosition() const;
+
+	/** Triggered whenever the viewport is clicked */
+	void OnViewportClicked(const FGeometry& MyGeometry, const FPointerEvent& PointerEvent);
+protected:
+	/**
+	 * Gets the displayed textures resolution as a string.
+	 *
+	 * @return Texture resolution string.
+	 */
+	FText GetDisplayedResolution() const;
+
+private:
+
+	// Callback for getting the zoom percentage text.
+	FText HandleZoomPercentageText( ) const;
+
+	// Checks if the texture being edited has a valid texture resource
+	bool HasValidTextureResource( ) const;
+private:
+	TWeakPtr<SSimulcamViewport> SimulcamViewportWeakPtr;
+	
+	// Level viewport client.
+	TSharedPtr<FSimulcamEditorViewportClient> ViewportClient;
+	// Slate viewport for rendering and IO.
+	TSharedPtr<FSceneViewport> Viewport;
+	// Viewport widget.
+	TSharedPtr<SViewport> ViewportWidget;
+
+	// Is rendering currently enabled? (disabled when reimporting a texture)
+	bool bIsRenderingEnabled;
+
+	/** The maximum width/height at which the texture will render in the preview window */
+	FVector2D CachedEffectiveTextureSize;
+
+	/** The texture's zoom factor. forced to 'fit' when 0 */
+	double Zoom = 0;
+
+public:
+	static constexpr double MaxZoom = 16.0;
+	static constexpr double ZoomStep = 0.1;
+};
diff --git a/Engine/Plugins/VirtualProduction/CameraCalibration/Source/CameraCalibrationEditor/Private/AssetEditor/Viewport/SimulcamEditorViewportClient.cpp b/Engine/Plugins/VirtualProduction/CameraCalibration/Source/CameraCalibrationEditor/Private/AssetEditor/Viewport/SimulcamEditorViewportClient.cpp
new file mode 100644
--- /dev/null
+++ b/Engine/Plugins/VirtualProduction/CameraCalibration/Source/CameraCalibrationEditor/Private/AssetEditor/Viewport/SimulcamEditorViewportClient.cpp
@@ -0,0 +1,370 @@
+// Copyright Epic Games, Inc. All Rights Reserved.
+
+#include "AssetEditor/Viewport/SimulcamEditorViewportClient.h"
+#include "AssetEditor/SSimulcamViewport.h"
+#include "AssetEditor/Viewport/SSimulcamEditorViewport.h"
+#include "CanvasItem.h"
+#include "CanvasTypes.h"
+#include "Engine/Texture2D.h"
+#include "Editor/UnrealEdEngine.h"
+#include "EngineModule.h"
+#include "EngineUtils.h"
+#include "ImageUtils.h"
+#include "Slate/SceneViewport.h"
+#include "Texture2DPreview.h"
+#include "UnrealEdGlobals.h"
+
+
+FSimulcamEditorViewportClient::FSimulcamEditorViewportClient(const TSharedRef<SSimulcamViewport>& InSimulcamViewport, const TSharedRef<SSimulcamEditorViewport>& InSimulcamEditorViewport, const bool bInWithZoom, const bool bInWithPan)
+	: SimulcamViewportWeakPtr(InSimulcamViewport)
+	, SimulcamEditorViewportWeakPtr(InSimulcamEditorViewport)
+	, bWithZoom(bInWithZoom)
+	, bWithPan(bInWithPan)
+{
+}
+
+void FSimulcamEditorViewportClient::Draw(FViewport* Viewport, FCanvas* Canvas)
+{
+	if (!SimulcamViewportWeakPtr.IsValid() || !SimulcamEditorViewportWeakPtr.IsValid())
+	{
+		return;
+	}
+
+	Canvas->Clear(FLinearColor::Black);
+
+	UTexture* Texture = SimulcamViewportWeakPtr.Pin()->GetTexture();
+	if (!Texture)
+	{
+		return;
+	}
+
+	SimulcamEditorViewportWeakPtr.Pin()->CacheEffectiveTextureSize();
+
+	if (SimulcamEditorViewportWeakPtr.Pin()->GetCustomZoomLevel() <= 0)
+	{
+		SimulcamEditorViewportWeakPtr.Pin()->SetCustomZoomLevel(-1);
+	}
+
+	// Figure out the size we need
+	FVector2D TextureSize = SimulcamEditorViewportWeakPtr.Pin()->CalculateTextureDimensions();
+
+	constexpr float MipLevel = 0;
+	constexpr float LayerIndex = 0;
+
+	TRefCountPtr<FBatchedElementParameters> BatchedElementParameters;
+	BatchedElementParameters = new FBatchedElementTexture2DPreviewParameters(MipLevel, LayerIndex, false, false, false, false, false);
+
+	if (Texture->GetResource() != nullptr)
+	{
+		FCanvasTileItem TileItem(CurrentTexturePosition, Texture->GetResource(), TextureSize, FLinearColor::White);
+		TileItem.BlendMode = ESimpleElementBlendMode::SE_BLEND_Opaque;
+		TileItem.BatchedElementParameters = BatchedElementParameters;
+		Canvas->DrawItem(TileItem);
+	}
+}
+
+void FSimulcamEditorViewportClient::MouseMove(FViewport* Viewport, int32 X, int32 Y)
+{
+	MousePosition = FIntPoint(X, Y);
+}
+
+void FSimulcamEditorViewportClient::ZoomOnPoint(FViewport* Viewport, FIntPoint InPoint, TFunction<void()> ZoomFunction)
+{
+	double CurrentZoom = SimulcamEditorViewportWeakPtr.Pin()->GetCustomZoomLevel();
+
+	ZoomFunction();
+
+	double NewZoom = SimulcamEditorViewportWeakPtr.Pin()->GetCustomZoomLevel();
+	FIntPoint ViewportSize = Viewport->GetSizeXY();
+
+	FVector2D TextureSize = SimulcamEditorViewportWeakPtr.Pin()->CalculateTextureDimensions();
+
+	if (TextureSize.IsNearlyZero())
+	{
+		return;
+	}
+
+	float XPercent = (InPoint.X - CurrentTexturePosition.X) / TextureSize.X;
+	float YPercent = (InPoint.Y - CurrentTexturePosition.Y) / TextureSize.Y;
+
+	CurrentTexturePosition.X = InPoint.X - (XPercent * (TextureSize.X * NewZoom / CurrentZoom));
+	CurrentTexturePosition.Y = InPoint.Y - (YPercent * (TextureSize.Y * NewZoom / CurrentZoom));
+}
+
+void FSimulcamEditorViewportClient::ZoomTowardsFit(FViewport* Viewport)
+{
+	double MinZoomLevel = SimulcamEditorViewportWeakPtr.Pin()->GetMinZoomLevel();
+
+	SimulcamEditorViewportWeakPtr.Pin()->ZoomOut();
+
+	double CurrentZoom = SimulcamEditorViewportWeakPtr.Pin()->GetCustomZoomLevel();
+	FIntPoint ViewportSize = Viewport->GetSizeXY();
+
+	FVector2D WantedTexturePosition = SimulcamEditorViewportWeakPtr.Pin()->GetFitPosition();
+
+	if (CurrentZoom >= MinZoomLevel)
+	{
+		double Gradient = 1.0f / ((CurrentZoom - MinZoomLevel) / SSimulcamEditorViewport::ZoomStep);
+		CurrentTexturePosition = FMath::Lerp(CurrentTexturePosition, WantedTexturePosition, FMath::Clamp(Gradient, 0.0, 1.0));
+	}
+}
+
+bool FSimulcamEditorViewportClient::InputKey(FViewport* Viewport, int32 ControllerId, FKey Key, EInputEvent Event, float AmountDepressed, bool Gamepad)
+{
+	if (Event == IE_Pressed)
+	{
+		if (!SimulcamEditorViewportWeakPtr.IsValid())
+		{
+			return false;
+		}
+
+		if (Key == EKeys::LeftMouseButton || Key == EKeys::MiddleMouseButton || Key == EKeys::RightMouseButton)
+		{
+			const FGeometry& MyGeometry = SimulcamEditorViewportWeakPtr.Pin()->GetTickSpaceGeometry();
+			const FVector2D LocalMouse = FVector2D(Viewport->GetMouseX(), Viewport->GetMouseY());
+			// check if we are under the viewport, otherwise the capture system will blindly trigger the PointerEvent
+			if (LocalMouse >= FVector2D(0, 0) && LocalMouse < MyGeometry.GetAbsoluteSize())
+			{
+				if (UTexture* Texture = SimulcamViewportWeakPtr.Pin()->GetTexture())
+				{
+					if (Texture->GetResource())
+					{
+						uint32 TextureWidth = Texture->GetResource()->GetSizeX();
+						uint32 TextureHeight = Texture->GetResource()->GetSizeY();
+
+						// create fake geometry and mouseposition
+						const FVector2D FakeMousePosition = GetTexturePosition();
+
+						// check for meaningful position
+						if (FakeMousePosition.X >= 0 && FakeMousePosition.Y >= 0 && FakeMousePosition.X < TextureWidth && FakeMousePosition.Y < TextureHeight)
+						{
+							const FGeometry FakeGeometry = FGeometry::MakeRoot(FVector2D(TextureWidth, TextureHeight), FSlateLayoutTransform());
+							FPointerEvent PointerEvent(
+								FSlateApplicationBase::CursorPointerIndex,
+								FakeMousePosition,
+								FakeMousePosition,
+								TSet<FKey>(),
+								Key,
+								0,
+								FSlateApplication::Get().GetPlatformApplication()->GetModifierKeys());
+
+							SimulcamEditorViewportWeakPtr.Pin()->OnViewportClicked(FakeGeometry, PointerEvent);
+						}
+					}
+				}
+			}
+		}
+
+		if (bWithZoom)
+		{
+			const bool bIsCtrlDown = FSlateApplication::Get().GetPlatformApplication()->GetModifierKeys().IsControlDown();
+
+			if (Key == EKeys::MouseScrollUp || (Key == EKeys::Add && bIsCtrlDown))
+			{
+				ZoomOnPoint(Viewport, MousePosition, [this] {SimulcamEditorViewportWeakPtr.Pin()->ZoomIn(); });
+				return true;
+			}
+
+			if (Key == EKeys::MouseScrollDown || (Key == EKeys::Subtract && bIsCtrlDown))
+			{	
+				FVector2D TextureSize = SimulcamEditorViewportWeakPtr.Pin()->CalculateTextureDimensions();
+				if (!bIsCtrlDown && (CurrentTexturePosition.X > 0 || (TextureSize.X + CurrentTexturePosition.X) < Viewport->GetSizeXY().X || CurrentTexturePosition.Y > 0 || (TextureSize.Y + CurrentTexturePosition.Y) < Viewport->GetSizeXY().Y))
+				{
+					ZoomTowardsFit(Viewport);
+				}
+				else
+				{
+					ZoomOnPoint(Viewport, MousePosition, [this] {SimulcamEditorViewportWeakPtr.Pin()->ZoomOut(); });
+				}
+				return true;
+			}
+
+			if ((Key == EKeys::Zero || Key == EKeys::NumPadZero) && bIsCtrlDown)
+			{
+				ZoomToFit(Viewport);
+				return true;
+			}
+		}
+	}
+
+	return false;
+}
+
+void FSimulcamEditorViewportClient::OnViewportResized(FViewport* InViewport, uint32 InParams)
+{
+	if (!SimulcamEditorViewportWeakPtr.IsValid())
+	{
+		return;
+	}
+
+	if (SimulcamEditorViewportWeakPtr.Pin()->GetViewport()->GetViewport() == InViewport)
+	{
+		FIntPoint NewViewportSize = InViewport->GetSizeXY();
+		if (NewViewportSize != CurrentViewportSize)
+		{
+			ZoomToFit(InViewport);
+			CurrentViewportSize = NewViewportSize;
+		}
+	}
+}
+
+void FSimulcamEditorViewportClient::ZoomToFit(FViewport* InViewport)
+{
+	if (SimulcamEditorViewportWeakPtr.IsValid())
+	{
+		SimulcamEditorViewportWeakPtr.Pin()->SetCustomZoomLevel(-1);
+		FVector2D TextureSize = SimulcamEditorViewportWeakPtr.Pin()->CalculateTextureDimensions();
+		if (!TextureSize.IsNearlyZero())
+		{
+			CurrentTexturePosition = FVector2D((InViewport->GetSizeXY().X - TextureSize.X) / 2.0f, (InViewport->GetSizeXY().Y - TextureSize.Y) / 2.0f);
+		}
+	}
+}
+
+bool FSimulcamEditorViewportClient::InputChar(FViewport* Viewport, int32 ControllerId, TCHAR Character)
+{
+	if (!bWithZoom)
+	{
+		return false;
+	}
+
+	if (!SimulcamViewportWeakPtr.IsValid())
+	{
+		return false;
+	}
+
+	const bool bIsCtrlDown = FSlateApplication::Get().GetPlatformApplication()->GetModifierKeys().IsControlDown();
+
+	if (Character == TCHAR('+') && bIsCtrlDown)
+	{
+		ZoomOnPoint(Viewport, MousePosition, [this] {SimulcamEditorViewportWeakPtr.Pin()->ZoomIn(); });
+		return true;
+	}
+
+	if (Character == TCHAR('-') && bIsCtrlDown)
+	{
+		ZoomOnPoint(Viewport, MousePosition, [this] {SimulcamEditorViewportWeakPtr.Pin()->ZoomOut(); });
+		return true;
+	}
+
+	return false;
+}
+
+bool FSimulcamEditorViewportClient::InputAxis(FViewport* Viewport, int32 ControllerId, FKey Key, float Delta, float DeltaTime, int32 NumSamples, bool bGamepad)
+{
+	if (Key == EKeys::MouseX || Key == EKeys::MouseY)
+	{
+		if (ShouldUseMousePanning(Viewport))
+		{
+			if (Key == EKeys::MouseY && CanPanVertically(Viewport, Delta))
+			{
+				CurrentTexturePosition.Y -= Delta;
+
+				FVector2D TextureSize = SimulcamEditorViewportWeakPtr.Pin()->CalculateTextureDimensions();
+
+				CurrentTexturePosition.Y = FMath::Clamp(CurrentTexturePosition.Y, Viewport->GetSizeXY().Y - TextureSize.Y, 0.0f);
+			}
+			else if (Key == EKeys::MouseX && CanPanHorizontally(Viewport, Delta))
+			{
+				CurrentTexturePosition.X += Delta;
+
+				FVector2D TextureSize = SimulcamEditorViewportWeakPtr.Pin()->CalculateTextureDimensions();
+
+				CurrentTexturePosition.X = FMath::Clamp(CurrentTexturePosition.X, Viewport->GetSizeXY().X - TextureSize.X, 0.0f);
+			}
+		}
+		return true;
+	}
+
+	return false;
+}
+
+bool FSimulcamEditorViewportClient::ShouldUseMousePanning(FViewport* Viewport) const
+{
+	return bWithPan && SimulcamViewportWeakPtr.IsValid() && SimulcamViewportWeakPtr.Pin()->GetTexture() && Viewport->KeyState(EKeys::RightMouseButton);
+}
+
+bool FSimulcamEditorViewportClient::CanPanHorizontally(FViewport* Viewport, float Direction) const
+{
+	FVector2D TextureSize = SimulcamEditorViewportWeakPtr.Pin()->CalculateTextureDimensions();
+	if (Direction < 0)
+	{
+		return (TextureSize.X + CurrentTexturePosition.X) > Viewport->GetSizeXY().X;
+	}
+	return CurrentTexturePosition.X < 0;
+}
+
+bool FSimulcamEditorViewportClient::CanPanVertically(FViewport* Viewport, float Direction) const
+{
+	FVector2D TextureSize = SimulcamEditorViewportWeakPtr.Pin()->CalculateTextureDimensions();
+	if (Direction > 0)
+	{
+		return (TextureSize.Y + CurrentTexturePosition.Y) > Viewport->GetSizeXY().Y;
+	}
+	return CurrentTexturePosition.Y < 0;
+}
+
+EMouseCursor::Type FSimulcamEditorViewportClient::GetCursor(FViewport* Viewport, int32 X, int32 Y)
+{
+	return ShouldUseMousePanning(Viewport) ? EMouseCursor::GrabHandClosed : EMouseCursor::Default;
+}
+
+FText FSimulcamEditorViewportClient::GetDisplayedResolution() const
+{
+	if (!SimulcamEditorViewportWeakPtr.IsValid() || !SimulcamViewportWeakPtr.IsValid() || !SimulcamViewportWeakPtr.Pin()->HasValidTextureResource())
+	{
+		return FText::GetEmpty();
+	}
+
+	FVector2D TextureSize = SimulcamEditorViewportWeakPtr.Pin()->CalculateTextureDimensions();
+
+	const FVector2D TexturePosition = GetTexturePosition();
+	const double CurrentZoom = SimulcamEditorViewportWeakPtr.Pin()->GetCustomZoomLevel();
+
+	uint32 TextureHeight = 1;
+	uint32 TextureWidth = 1;
+	UTexture* Texture = SimulcamViewportWeakPtr.Pin()->GetTexture();
+
+	if (Texture && Texture->GetResource())
+	{
+		TextureWidth = Texture->GetResource()->GetSizeX();
+		TextureHeight = Texture->GetResource()->GetSizeY();
+	}
+
+	int32 TexturePositionX = TexturePosition.X;
+	if (TexturePositionX < 0)
+	{
+		TexturePositionX = 0;
+	}
+	else if (TexturePositionX > static_cast<int32>(TextureWidth))
+	{
+		TexturePositionX = TextureWidth;
+	}
+
+	int32 TexturePositionY = TexturePosition.Y;
+	if (TexturePositionY < 0)
+	{
+		TexturePositionY = 0;
+	}
+	else if (TexturePositionY > static_cast<int32>(TextureHeight))
+	{
+		TexturePositionY = TextureHeight;
+	}
+
+	return FText::Format(
+		FText::FromString("Displayed: {0}x{1}\nTextureSize: {2}x{3}\nTexturePosition: {4}x{5}\nZoom: {6}"),
+		FText::AsNumber(FMath::Max(1.0f, TextureSize.X)),
+		FText::AsNumber(FMath::Max(1.0f, TextureSize.Y)),
+		FText::AsNumber(TextureWidth),
+		FText::AsNumber(TextureHeight),
+		FText::AsNumber(TexturePositionX),
+		FText::AsNumber(TexturePositionY),
+		FText::AsNumber(CurrentZoom)
+	);
+}
+
+FVector2D FSimulcamEditorViewportClient::GetTexturePosition() const
+{
+	const double CurrentZoom = SimulcamEditorViewportWeakPtr.Pin()->GetCustomZoomLevel();
+	return FVector2D((MousePosition.X - CurrentTexturePosition.X) / CurrentZoom, (MousePosition.Y - CurrentTexturePosition.Y) / CurrentZoom);
+}
diff --git a/Engine/Plugins/VirtualProduction/CameraCalibration/Source/CameraCalibrationEditor/Private/AssetEditor/Viewport/SimulcamEditorViewportClient.h b/Engine/Plugins/VirtualProduction/CameraCalibration/Source/CameraCalibrationEditor/Private/AssetEditor/Viewport/SimulcamEditorViewportClient.h
new file mode 100644
--- /dev/null
+++ b/Engine/Plugins/VirtualProduction/CameraCalibration/Source/CameraCalibrationEditor/Private/AssetEditor/Viewport/SimulcamEditorViewportClient.h
@@ -0,0 +1,67 @@
+// Copyright Epic Games, Inc. All Rights Reserved.
+#pragma once
+
+#include "UnrealClient.h"
+#include "CoreMinimal.h"
+#include "InputCoreTypes.h"
+#include "UObject/GCObject.h"
+
+class FCanvas;
+class ITextureEditorToolkit;
+class STextureEditorViewport;
+class UTexture2D;
+class SSimulcamViewport;
+class SSimulcamEditorViewport;
+
+class FSimulcamEditorViewportClient
+	: public FViewportClient
+{
+public:
+	/** Constructor */
+	FSimulcamEditorViewportClient(const TSharedRef<SSimulcamViewport>& InSimulcamViewport, const TSharedRef<SSimulcamEditorViewport>& InSimulcamEditorViewport, const bool bInWithZoom, const bool bInWithPan);
+
+	/** Begin FViewportClient interface */
+	virtual void Draw(FViewport* Viewport, FCanvas* Canvas) override;
+	virtual bool InputKey(FViewport* Viewport, int32 ControllerId, FKey Key, EInputEvent Event, float AmountDepressed = 1.0f, bool bGamepad = false) override;
+	virtual bool InputChar(FViewport* Viewport, int32 ControllerId, TCHAR Character) override;
+	virtual bool InputAxis(FViewport* Viewport, int32 ControllerId, FKey Key, float Delta, float DeltaTime, int32 NumSamples=1, bool bGamepad=false) override;
+	virtual UWorld* GetWorld() const override { return nullptr; }
+	virtual EMouseCursor::Type GetCursor(FViewport* Viewport, int32 X, int32 Y) override;
+	virtual void MouseMove(FViewport* Viewport, int32 X, int32 Y) override;
+	/** End FViewportClient interface */
+
+	/** Returns a string representation of the currently displayed textures resolution */
+	FText GetDisplayedResolution() const;
+
+	/** Triggered whenever the viewport is resized */
+	void OnViewportResized(FViewport* InViewport, uint32 InParams);
+private:
+	FVector2D GetTexturePosition() const;
+	
+	/** TRUE if right clicking and dragging for panning a texture 2D */
+	bool ShouldUseMousePanning(FViewport* Viewport) const;
+
+	bool CanPanHorizontally(FViewport* Viewport, float Direction) const;
+	bool CanPanVertically(FViewport* Viewport, float Direction) const;
+
+	/** Zoom under a specific pixel maintaining it aligned to the cursor */
+	void ZoomOnPoint(FViewport* Viewport, FIntPoint InPoint, TFunction<void()> ZoomFunction);
+
+	/** Zoom to fit the texture in the viewport */
+	void ZoomToFit(FViewport* InViewport);
+
+	/** Zoom Out trying to (smoothly) re-fit the image if there is space on borders */
+	void ZoomTowardsFit(FViewport* Viewport);
+
+	TWeakPtr<SSimulcamViewport> SimulcamViewportWeakPtr;
+	TWeakPtr<SSimulcamEditorViewport> SimulcamEditorViewportWeakPtr;
+
+	FIntPoint MousePosition;
+
+	FVector2D CurrentTexturePosition = FVector2D(0, 0);
+
+	FIntPoint CurrentViewportSize = FIntPoint(0, 0);
+
+	bool bWithZoom;
+	bool bWithPan;
+};
diff --git a/Engine/Plugins/VirtualProduction/CameraCalibration/Source/CameraCalibrationEditor/Private/Calibrators/CameraLensDistortionAlgoCheckerboard.cpp b/Engine/Plugins/VirtualProduction/CameraCalibration/Source/CameraCalibrationEditor/Private/Calibrators/CameraLensDistortionAlgoCheckerboard.cpp
--- a/Engine/Plugins/VirtualProduction/CameraCalibration/Source/CameraCalibrationEditor/Private/Calibrators/CameraLensDistortionAlgoCheckerboard.cpp
+++ b/Engine/Plugins/VirtualProduction/CameraCalibration/Source/CameraCalibrationEditor/Private/Calibrators/CameraLensDistortionAlgoCheckerboard.cpp
@@ -342,7 +342,7 @@ bool UCameraLensDistortionAlgoCheckerboard::AddCalibrationRow(FText& OutErrorMes
 
 		if (UTexture2D* ThumbnailTexture = FOpenCVHelper::TextureFromCvMat(CvThumbnail))
 		{
-			Row->Thumbnail = SNew(SSimulcamViewport, ThumbnailTexture);
+			Row->Thumbnail = SNew(SSimulcamViewport, ThumbnailTexture).WithZoom(false).WithPan(false);
 		}
 	}
 
