[CameraCalibration] - Add new Image Center Tool calibration step to the LensFile editor

diff --git a/Engine/Plugins/VirtualProduction/CameraCalibration/Source/CameraCalibrationEditor/Private/AssetEditor/CameraCalibrationStepsController.cpp b/Engine/Plugins/VirtualProduction/CameraCalibration/Source/CameraCalibrationEditor/Private/AssetEditor/CameraCalibrationStepsController.cpp
--- a/Engine/Plugins/VirtualProduction/CameraCalibration/Source/CameraCalibrationEditor/Private/AssetEditor/CameraCalibrationStepsController.cpp
+++ b/Engine/Plugins/VirtualProduction/CameraCalibration/Source/CameraCalibrationEditor/Private/AssetEditor/CameraCalibrationStepsController.cpp
@@ -32,8 +32,8 @@
 #include "Materials/Material.h"
 #include "Materials/MaterialInstanceDynamic.h"
 #include "MediaPlayer.h"
-#include "MediaTexture.h"
 #include "MediaSource.h"
+#include "MediaTexture.h"
 #include "Misc/MessageDialog.h"
 #include "Models/SphericalLensModel.h"
 #include "Modules/ModuleManager.h"
@@ -773,6 +773,22 @@ void FCameraCalibrationStepsController::OnSimulcamViewportClicked(const FGeometr
 	}
 }
 
+bool FCameraCalibrationStepsController::OnSimulcamViewportInputKey(const FKey& InKey, const EInputEvent& InEvent)
+{
+	bool bStepHandled = false;
+
+	for (TStrongObjectPtr<UCameraCalibrationStep>& Step : CalibrationSteps)
+	{
+		if (Step.IsValid() && Step->IsActive())
+		{
+			bStepHandled |= Step->OnViewportInputKey(InKey, InEvent);
+			break;
+		}
+	}
+
+	return bStepHandled;
+}
+
 ULiveLinkCameraController* FCameraCalibrationStepsController::FindLiveLinkCameraController() const
 {
 	return FindLiveLinkCameraControllerWithLens(Camera.Get(), LensFile.Get());
@@ -979,7 +995,7 @@ void FCameraCalibrationStepsController::FindMediaSourceUrls(TArray<TSharedPtr<FS
 	{
 		if (const UMediaSource* MediaSource = MediaProfile->GetMediaSource(MediaSourceIdx))
 		{
-			OutMediaSourceUrls.Add(MakeShareable(new FString(MediaSource->GetUrl())));
+			OutMediaSourceUrls.Add(MakeShared<FString>(MediaSource->GetUrl()));
 		}
 	}
 }
diff --git a/Engine/Plugins/VirtualProduction/CameraCalibration/Source/CameraCalibrationEditor/Private/AssetEditor/CameraCalibrationStepsController.h b/Engine/Plugins/VirtualProduction/CameraCalibration/Source/CameraCalibrationEditor/Private/AssetEditor/CameraCalibrationStepsController.h
--- a/Engine/Plugins/VirtualProduction/CameraCalibration/Source/CameraCalibrationEditor/Private/AssetEditor/CameraCalibrationStepsController.h
+++ b/Engine/Plugins/VirtualProduction/CameraCalibration/Source/CameraCalibrationEditor/Private/AssetEditor/CameraCalibrationStepsController.h
@@ -123,6 +123,9 @@ public:
 	/** Called by the UI when the Simulcam Viewport is clicked */
 	void OnSimulcamViewportClicked(const FGeometry& MyGeometry, const FPointerEvent& MouseEvent);
 
+	/** Called by the UI when the Simulcam Viewport receives keyboard input */
+	bool OnSimulcamViewportInputKey(const FKey& InKey, const EInputEvent& InEvent);
+
 private:
 
 	/** Finds the LiveLinkCameraController used in the given CameraActor that is also using the given LensFile */
diff --git a/Engine/Plugins/VirtualProduction/CameraCalibration/Source/CameraCalibrationEditor/Private/AssetEditor/ImageCenterTool.cpp b/Engine/Plugins/VirtualProduction/CameraCalibration/Source/CameraCalibrationEditor/Private/AssetEditor/ImageCenterTool.cpp
new file mode 100644
--- /dev/null
+++ b/Engine/Plugins/VirtualProduction/CameraCalibration/Source/CameraCalibrationEditor/Private/AssetEditor/ImageCenterTool.cpp
@@ -0,0 +1,157 @@
+// Copyright Epic Games, Inc. All Rights Reserved.
+
+#include "ImageCenterTool.h"
+
+#include "CameraCalibrationStepsController.h"
+#include "CameraCalibrationSubsystem.h"
+#include "CameraImageCenterAlgo.h"
+#include "LensDistortionTool.h"
+#include "LensFile.h"
+#include "Misc/MessageDialog.h"
+#include "ScopedTransaction.h"
+#include "SImageCenterToolPanel.h"
+
+#define LOCTEXT_NAMESPACE "ImageCenterTool"
+
+void UImageCenterTool::Initialize(TWeakPtr<FCameraCalibrationStepsController> InCameraCalibrationStepController)
+{
+	CameraCalibrationStepsController = InCameraCalibrationStepController;
+}
+
+void UImageCenterTool::Shutdown()
+{
+	if (CurrentAlgo)
+	{
+		CurrentAlgo->Shutdown();
+		CurrentAlgo = nullptr;
+	}
+}
+
+void UImageCenterTool::Tick(float DeltaTime)
+{
+	if (CurrentAlgo && CurrentAlgo->IsActive())
+	{
+		CurrentAlgo->Tick(DeltaTime);
+	}
+}
+
+bool UImageCenterTool::OnViewportClicked(const FGeometry& MyGeometry, const FPointerEvent& MouseEvent)
+{
+	if (IsActive() && CurrentAlgo && CurrentAlgo->IsActive())
+	{
+		return CurrentAlgo->OnViewportClicked(MyGeometry, MouseEvent);
+	}
+
+	return false;
+}
+
+bool UImageCenterTool::OnViewportInputKey(const FKey& InKey, const EInputEvent& InEvent)
+{
+	if (IsActive() && CurrentAlgo && CurrentAlgo->IsActive())
+	{
+		return CurrentAlgo->OnViewportInputKey(InKey, InEvent);
+	}
+
+	return false;
+}
+
+TSharedRef<SWidget> UImageCenterTool::BuildUI()
+{
+	return SNew(SImageCenterToolPanel, this);
+}
+
+bool UImageCenterTool::DependsOnStep(UCameraCalibrationStep* Step) const
+{
+	return Cast<ULensDistortionTool>(Step) != nullptr;
+}
+
+void UImageCenterTool::Activate()
+{
+	bIsActive = true;
+
+	if (CurrentAlgo && !CurrentAlgo->IsActive())
+	{
+		CurrentAlgo->Activate();
+	}
+}
+
+void UImageCenterTool::Deactivate()
+{
+	bIsActive = false;
+
+	if (CurrentAlgo && CurrentAlgo->IsActive())
+	{
+		CurrentAlgo->Deactivate();
+	}
+}
+
+bool UImageCenterTool::IsActive() const
+{
+	return bIsActive;
+}
+
+FCameraCalibrationStepsController* UImageCenterTool::GetCameraCalibrationStepsController() const
+{
+	return CameraCalibrationStepsController.Pin().Get();
+}
+
+UCameraImageCenterAlgo* UImageCenterTool::GetAlgo() const
+{
+	return CurrentAlgo;
+}
+
+void UImageCenterTool::SetAlgo(const FName& AlgoName)
+{
+	// Query subsystem for the selected algorithm class
+	UCameraCalibrationSubsystem* Subsystem = GEngine->GetEngineSubsystem<UCameraCalibrationSubsystem>();
+	check(Subsystem);
+
+	TSubclassOf<UCameraImageCenterAlgo> AlgoClass = Subsystem->GetCameraImageCenterAlgo(AlgoName);
+
+	// If it is the same as the existing one, do nothing.
+	if (!CurrentAlgo && !AlgoClass)
+	{
+		return;
+	}
+	
+	if (CurrentAlgo && (CurrentAlgo->GetClass() == AlgoClass))
+	{
+		return;
+	}
+
+	// Remove old Algo
+	if (CurrentAlgo)
+	{
+		CurrentAlgo->Shutdown();
+		CurrentAlgo = nullptr;
+	}
+
+	// If AlgoClass is none, we're done here.
+	if (!AlgoClass)
+	{
+		return;
+	}
+
+	// Create new algo
+	CurrentAlgo = NewObject<UCameraImageCenterAlgo>(
+		GetTransientPackage(),
+		AlgoClass,
+		MakeUniqueObjectName(GetTransientPackage(), AlgoClass), 
+		RF_Transactional);
+
+	if (CurrentAlgo)
+	{
+		CurrentAlgo->Initialize(this);
+	}
+}
+
+void UImageCenterTool::OnSaveCurrentImageCenter()
+{
+	if (CurrentAlgo && CurrentAlgo->IsActive() && CurrentAlgo->HasImageCenterChanged())
+	{
+		CurrentAlgo->OnSavedImageCenter();
+	}
+}
+
+
+#undef LOCTEXT_NAMESPACE
diff --git a/Engine/Plugins/VirtualProduction/CameraCalibration/Source/CameraCalibrationEditor/Private/AssetEditor/ImageCenterTool.h b/Engine/Plugins/VirtualProduction/CameraCalibration/Source/CameraCalibrationEditor/Private/AssetEditor/ImageCenterTool.h
new file mode 100644
--- /dev/null
+++ b/Engine/Plugins/VirtualProduction/CameraCalibration/Source/CameraCalibrationEditor/Private/AssetEditor/ImageCenterTool.h
@@ -0,0 +1,60 @@
+// Copyright Epic Games, Inc. All Rights Reserved.
+
+#pragma once
+
+#include "CameraCalibrationStep.h"
+
+#include "ImageCenterTool.generated.h"
+
+struct FLensFileEvalData;
+
+class UCameraImageCenterAlgo;
+
+/**
+ * UImageCenterTool is the controller for the image center panel.
+ */
+UCLASS()
+class UImageCenterTool : public UCameraCalibrationStep
+{
+	GENERATED_BODY()
+
+public:
+
+	//~ Begin UCameraCalibrationStep interface
+	virtual void Initialize(TWeakPtr<FCameraCalibrationStepsController> InCameraCalibrationStepController) override;
+	virtual void Shutdown() override;
+	virtual void Tick(float DeltaTime) override;
+	virtual bool OnViewportClicked(const FGeometry& MyGeometry, const FPointerEvent& MouseEvent) override;
+	virtual bool OnViewportInputKey(const FKey& InKey, const EInputEvent& InEvent) override;
+	virtual TSharedRef<SWidget> BuildUI() override;
+	virtual FName FriendlyName() const  override { return TEXT("Image Center"); };
+	virtual bool DependsOnStep(UCameraCalibrationStep* Step) const override;
+	virtual void Activate() override;
+	virtual void Deactivate() override;
+	virtual bool IsActive() const override;
+	virtual FCameraCalibrationStepsController* GetCameraCalibrationStepsController() const override;
+	//~ End UCameraCalibrationStep interface
+
+	/** Returns the currently active algorithm */
+	UCameraImageCenterAlgo* GetAlgo() const;
+
+	/** Sets the active algorithm by name */
+	void SetAlgo(const FName& AlgoName);
+
+public:
+
+	/** Called by the UI when the user wants to save the calibration data that the current algorithm is providing */
+	void OnSaveCurrentImageCenter();
+
+private:
+
+	/** Pointer to the calibration steps controller */
+	TWeakPtr<FCameraCalibrationStepsController> CameraCalibrationStepsController;
+
+	/** The currently selected algorithm */
+	UPROPERTY(Transient)
+	TObjectPtr<UCameraImageCenterAlgo> CurrentAlgo;
+
+	/** True if this tool is the active one in the panel */
+	bool bIsActive = false;
+};
diff --git a/Engine/Plugins/VirtualProduction/CameraCalibration/Source/CameraCalibrationEditor/Private/AssetEditor/LensDistortionTool.h b/Engine/Plugins/VirtualProduction/CameraCalibration/Source/CameraCalibrationEditor/Private/AssetEditor/LensDistortionTool.h
--- a/Engine/Plugins/VirtualProduction/CameraCalibration/Source/CameraCalibrationEditor/Private/AssetEditor/LensDistortionTool.h
+++ b/Engine/Plugins/VirtualProduction/CameraCalibration/Source/CameraCalibrationEditor/Private/AssetEditor/LensDistortionTool.h
@@ -64,7 +64,7 @@ private:
 	UPROPERTY(Transient)
 	UCameraLensDistortionAlgo* CurrentAlgo;
 
-	/** Holds the registered camera nodal offset algos */
+	/** Holds the registered camera lens distortion algos */
 	UPROPERTY(Transient)
 	TMap<FName, TSubclassOf<UCameraLensDistortionAlgo>> AlgosMap;
 
diff --git a/Engine/Plugins/VirtualProduction/CameraCalibration/Source/CameraCalibrationEditor/Private/AssetEditor/NodalOffsetTool.cpp b/Engine/Plugins/VirtualProduction/CameraCalibration/Source/CameraCalibrationEditor/Private/AssetEditor/NodalOffsetTool.cpp
--- a/Engine/Plugins/VirtualProduction/CameraCalibration/Source/CameraCalibrationEditor/Private/AssetEditor/NodalOffsetTool.cpp
+++ b/Engine/Plugins/VirtualProduction/CameraCalibration/Source/CameraCalibrationEditor/Private/AssetEditor/NodalOffsetTool.cpp
@@ -6,7 +6,7 @@
 #include "CameraCalibrationStepsController.h"
 #include "CameraCalibrationSubsystem.h"
 #include "CameraCalibrationTypes.h"
-#include "LensDistortionTool.h"
+#include "ImageCenterTool.h"
 #include "LiveLinkCameraController.h"
 #include "Misc/MessageDialog.h"
 #include "ScopedTransaction.h"
@@ -58,7 +58,7 @@ TSharedRef<SWidget> UNodalOffsetTool::BuildUI()
 
 bool UNodalOffsetTool::DependsOnStep(UCameraCalibrationStep* Step) const
 {
-	return !!Cast<ULensDistortionTool>(Step);
+	return !!Cast<UImageCenterTool>(Step);
 }
 
 void UNodalOffsetTool::Activate()
diff --git a/Engine/Plugins/VirtualProduction/CameraCalibration/Source/CameraCalibrationEditor/Private/AssetEditor/SCameraCalibrationSteps.cpp b/Engine/Plugins/VirtualProduction/CameraCalibration/Source/CameraCalibrationEditor/Private/AssetEditor/SCameraCalibrationSteps.cpp
--- a/Engine/Plugins/VirtualProduction/CameraCalibration/Source/CameraCalibrationEditor/Private/AssetEditor/SCameraCalibrationSteps.cpp
+++ b/Engine/Plugins/VirtualProduction/CameraCalibration/Source/CameraCalibrationEditor/Private/AssetEditor/SCameraCalibrationSteps.cpp
@@ -68,6 +68,7 @@ void SCameraCalibrationSteps::Construct(const FArguments& InArgs, TWeakPtr<FCame
 			[
 				SNew(SSimulcamViewport, CalibrationStepsController.Pin()->GetRenderTarget())
 				.OnSimulcamViewportClicked_Raw(CalibrationStepsController.Pin().Get(), &FCameraCalibrationStepsController::OnSimulcamViewportClicked)
+				.OnSimulcamViewportInputKey_Raw(CalibrationStepsController.Pin().Get(), &FCameraCalibrationStepsController::OnSimulcamViewportInputKey)
 			]
 		]
 
@@ -226,7 +227,7 @@ void SCameraCalibrationSteps::UpdateMediaSourcesOptions()
 	}
 
 	// Add a "None" option
-	CurrentMediaSources.Add(MakeShareable(new FString(TEXT("None"))));
+	CurrentMediaSources.Add(MakeShared<FString>(TEXT("None")));
 
 	check(MediaSourcesComboBox.IsValid());
 
diff --git a/Engine/Plugins/VirtualProduction/CameraCalibration/Source/CameraCalibrationEditor/Private/AssetEditor/SImageCenterToolPanel.cpp b/Engine/Plugins/VirtualProduction/CameraCalibration/Source/CameraCalibrationEditor/Private/AssetEditor/SImageCenterToolPanel.cpp
new file mode 100644
--- /dev/null
+++ b/Engine/Plugins/VirtualProduction/CameraCalibration/Source/CameraCalibrationEditor/Private/AssetEditor/SImageCenterToolPanel.cpp
@@ -0,0 +1,215 @@
+// Copyright Epic Games, Inc. All Rights Reserved.
+
+#include "SImageCenterToolPanel.h"
+
+#include "CameraCalibrationSubsystem.h"
+#include "Dialogs/CustomDialog.h"
+#include "EditorFontGlyphs.h"
+#include "ImageCenterTool.h"
+#include "UI/CameraCalibrationWidgetHelpers.h"
+#include "Widgets/Input/SButton.h"
+#include "Widgets/Input/SComboBox.h"
+#include "Widgets/Text/STextBlock.h"
+
+
+#define LOCTEXT_NAMESPACE "ImageCenterTool"
+
+
+void SImageCenterToolPanel::Construct(const FArguments& InArgs, UImageCenterTool* InImageCenterTool)
+{
+	ImageCenterTool = InImageCenterTool;
+
+	// This will be the widget wrapper of the custom algo UI.
+	ImageCenterUI = SNew(SVerticalBox);
+
+	ChildSlot
+	[
+		SNew(SHorizontalBox)
+
+		+ SHorizontalBox::Slot() // Right toolbar
+		.FillWidth(0.25f)
+		[
+			SNew(SVerticalBox)
+
+			+ SVerticalBox::Slot() // Algo picker
+			.MaxHeight(FCameraCalibrationWidgetHelpers::DefaultRowHeight)
+			[FCameraCalibrationWidgetHelpers::BuildLabelWidgetPair(LOCTEXT("ImageCenterAlgo", "Image Center Algo"), BuildImageCenterAlgoPickerWidget())]
+
+			+ SVerticalBox::Slot() // Algo UI
+			.AutoHeight()
+			[BuildImageCenterUIWrapper()]
+
+			+ SVerticalBox::Slot() // Save Image Center
+			.AutoHeight()
+			.Padding(0, 20)
+			[
+				SNew(SButton).Text(LOCTEXT("AddToImageCenterLUT", "Update Image Center Calibration"))
+				.HAlign(HAlign_Center)
+				.VAlign(VAlign_Center)
+				.OnClicked_Lambda([&]() -> FReply
+				{
+					if (ImageCenterTool.IsValid())
+					{
+						ImageCenterTool->OnSaveCurrentImageCenter();
+					}
+					return FReply::Handled();
+				})
+			]
+		]
+	];
+}
+
+TSharedRef<SWidget> SImageCenterToolPanel::BuildImageCenterUIWrapper()
+{
+	return SNew(SVerticalBox)
+
+		+ SVerticalBox::Slot() // Algo's Widget
+		.AutoHeight()
+		[ImageCenterUI.ToSharedRef()];
+}
+
+void SImageCenterToolPanel::UpdateImageCenterUI()
+{
+	check(AlgosComboBox.IsValid());
+
+	// Get current algo to later compare with new one
+	const UCameraImageCenterAlgo* OldAlgo = ImageCenterTool->GetAlgo();
+
+	// Set new algo by name
+	const FName AlgoName(*AlgosComboBox->GetSelectedItem());
+	ImageCenterTool->SetAlgo(AlgoName);
+
+	// Get the new algo
+	UCameraImageCenterAlgo* Algo = ImageCenterTool->GetAlgo();
+
+	// nullptr may indicate that it was unregistered, so refresh combobox options.
+	if (!Algo)
+	{
+		UpdateAlgosOptions();
+		return;
+	}
+
+	// If we didn't change the algo, we're done here.
+	if (Algo == OldAlgo)
+	{
+		return;
+	}
+
+	// Remove old UI
+	check(ImageCenterUI.IsValid());
+	ImageCenterUI->ClearChildren();
+
+	// Assign GUI
+	ImageCenterUI->AddSlot()[Algo->BuildUI()];
+}
+
+void SImageCenterToolPanel::UpdateAlgosOptions()
+{
+	CurrentAlgos.Empty();
+
+	// Ask the subsystem for the list of registered Algos
+	UCameraCalibrationSubsystem* Subsystem = GEngine->GetEngineSubsystem<UCameraCalibrationSubsystem>();
+	check(Subsystem);
+
+	for (FName& Name : Subsystem->GetCameraImageCenterAlgos())
+	{
+		CurrentAlgos.Add(MakeShared<FString>(Name.ToString()));
+	}
+
+	// Ask the ComboBox to refresh its options from its source (that we just updated)
+	AlgosComboBox->RefreshOptions();
+}
+
+TSharedRef<SWidget> SImageCenterToolPanel::BuildImageCenterAlgoPickerWidget()
+{
+	// Create ComboBox widget
+
+	AlgosComboBox = SNew(SComboBox<TSharedPtr<FString>>)
+		.OptionsSource(&CurrentAlgos)
+		.OnSelectionChanged_Lambda([&](TSharedPtr<FString> NewValue, ESelectInfo::Type Type) -> void
+		{
+			// Replace the custom algo widget
+			UpdateImageCenterUI();
+		})
+		.OnGenerateWidget_Lambda([&](TSharedPtr<FString> InOption) -> TSharedRef<SWidget>
+		{
+			return SNew(STextBlock).Text(FText::FromString(*InOption));
+		})
+		.InitiallySelectedItem(nullptr)
+		[
+			SNew(STextBlock)
+			.Text_Lambda([&]() -> FText
+			{
+				if (AlgosComboBox.IsValid() && AlgosComboBox->GetSelectedItem().IsValid())
+				{
+					return FText::FromString(*AlgosComboBox->GetSelectedItem());
+				}
+
+				return LOCTEXT("InvalidComboOption", "Invalid");
+			})
+		];
+
+	// Update the object holding this combobox's options source
+	UpdateAlgosOptions();
+
+	// Pick the first option by default (if available)
+	if (CurrentAlgos.Num())
+	{
+		AlgosComboBox->SetSelectedItem(CurrentAlgos[0]);
+	}
+	else
+	{
+		AlgosComboBox->SetSelectedItem(nullptr);
+	}
+
+	return SNew(SHorizontalBox)
+
+		+ SHorizontalBox::Slot() // algo picker
+		[
+			AlgosComboBox.ToSharedRef()
+		]
+
+		+ SHorizontalBox::Slot() // Help button
+		.AutoWidth()
+		[
+			SNew(SButton)
+			.ToolTipText(LOCTEXT("ShowHelp_Tip", "Help about this algo"))
+			.HAlign(HAlign_Center)
+			.VAlign(VAlign_Center)
+			.ButtonStyle(FEditorStyle::Get(), "HoverHintOnly")
+			.OnClicked_Lambda([&]() -> FReply
+			{
+				if (!ImageCenterTool.IsValid())
+				{
+					return FReply::Handled();
+				}
+
+				UCameraImageCenterAlgo* Algo = ImageCenterTool->GetAlgo();
+
+				if (!Algo)
+				{
+					return FReply::Handled();
+				}
+
+				TSharedRef<SCustomDialog> AlgoHelpWindow =
+					SNew(SCustomDialog)
+					.Title(FText::FromName(ImageCenterTool->FriendlyName()))
+					.DialogContent(Algo->BuildHelpWidget())
+					.Buttons({ SCustomDialog::FButton(LOCTEXT("Ok", "Ok")) });
+
+				AlgoHelpWindow->Show();
+
+				return FReply::Handled();
+			})
+			[
+				SNew(STextBlock)
+				.Font(FEditorStyle::Get().GetFontStyle("FontAwesome.12"))
+				.Text(FEditorFontGlyphs::Info_Circle)
+				.ColorAndOpacity(FLinearColor::White)
+			]
+		]
+	;
+}
+
+
+#undef LOCTEXT_NAMESPACE
diff --git a/Engine/Plugins/VirtualProduction/CameraCalibration/Source/CameraCalibrationEditor/Private/AssetEditor/SImageCenterToolPanel.h b/Engine/Plugins/VirtualProduction/CameraCalibration/Source/CameraCalibrationEditor/Private/AssetEditor/SImageCenterToolPanel.h
new file mode 100644
--- /dev/null
+++ b/Engine/Plugins/VirtualProduction/CameraCalibration/Source/CameraCalibrationEditor/Private/AssetEditor/SImageCenterToolPanel.h
@@ -0,0 +1,55 @@
+// Copyright Epic Games, Inc. All Rights Reserved.
+#pragma once
+
+#include "CoreMinimal.h"
+ 
+#include "Widgets/DeclarativeSyntaxSupport.h"
+#include "Widgets/SCompoundWidget.h"
+
+template<typename OptionType>
+class SComboBox;
+
+class UImageCenterTool;
+
+/**
+ * UI for the image center adjustment tool.
+ * It also holds the UI given by the selected image center algorithm.
+ */
+class SImageCenterToolPanel : public SCompoundWidget
+{
+	SLATE_BEGIN_ARGS(SImageCenterToolPanel) {}
+	SLATE_END_ARGS()
+
+public:
+	void Construct(const FArguments& InArgs, UImageCenterTool* InImageCenterTool);
+
+private:
+
+	/** Builds the wrapper for the currently selected iamge center UI */
+	TSharedRef<SWidget> BuildImageCenterUIWrapper();
+
+	/** Builds the UI for the image center algorithm picker */
+	TSharedRef<SWidget> BuildImageCenterAlgoPickerWidget();
+
+	/** Updates the UI so that it matches the selected image center algorithm (if necessary) */
+	void UpdateImageCenterUI();
+
+	/** Refreshes the list of available image center algorithms shown in the AlgosComboBox */
+	void UpdateAlgosOptions();
+
+private:
+
+	/** The image center tool controller object */
+	TWeakObjectPtr<class UImageCenterTool> ImageCenterTool;
+
+	/** The box containing the UI given by the selected image center algorithm */
+	TSharedPtr<class SVerticalBox> ImageCenterUI;
+
+private:
+
+	/** Options source for the AlgosComboBox. Lists the currently available image center algorithms */
+	TArray<TSharedPtr<FString>> CurrentAlgos;
+
+	/** The combobox that presents the available image center algorithms */
+	TSharedPtr<SComboBox<TSharedPtr<FString>>> AlgosComboBox;
+};
diff --git a/Engine/Plugins/VirtualProduction/CameraCalibration/Source/CameraCalibrationEditor/Private/AssetEditor/SLensDistortionToolPanel.cpp b/Engine/Plugins/VirtualProduction/CameraCalibration/Source/CameraCalibrationEditor/Private/AssetEditor/SLensDistortionToolPanel.cpp
--- a/Engine/Plugins/VirtualProduction/CameraCalibration/Source/CameraCalibrationEditor/Private/AssetEditor/SLensDistortionToolPanel.cpp
+++ b/Engine/Plugins/VirtualProduction/CameraCalibration/Source/CameraCalibrationEditor/Private/AssetEditor/SLensDistortionToolPanel.cpp
@@ -125,7 +125,7 @@ void SLensDistortionToolPanel::UpdateAlgosOptions()
 
 	for (FName& AlgoName : Tool->GetAlgos())
 	{
-		CurrentAlgos.Add(MakeShareable(new FString(AlgoName.ToString())));
+		CurrentAlgos.Add(MakeShared<FString>(AlgoName.ToString()));
 	}
 
 	// Ask the ComboBox to refresh its options from its source (that we just updated)
diff --git a/Engine/Plugins/VirtualProduction/CameraCalibration/Source/CameraCalibrationEditor/Private/AssetEditor/SNodalOffsetToolPanel.cpp b/Engine/Plugins/VirtualProduction/CameraCalibration/Source/CameraCalibrationEditor/Private/AssetEditor/SNodalOffsetToolPanel.cpp
--- a/Engine/Plugins/VirtualProduction/CameraCalibration/Source/CameraCalibrationEditor/Private/AssetEditor/SNodalOffsetToolPanel.cpp
+++ b/Engine/Plugins/VirtualProduction/CameraCalibration/Source/CameraCalibrationEditor/Private/AssetEditor/SNodalOffsetToolPanel.cpp
@@ -124,7 +124,7 @@ void SNodalOffsetToolPanel::UpdateAlgosOptions()
 
 	for (FName& Name : Subsystem->GetCameraNodalOffsetAlgos())
 	{
-		CurrentAlgos.Add(MakeShareable(new FString(Name.ToString())));
+		CurrentAlgos.Add(MakeShared<FString>(Name.ToString()));
 	}
 
 	// Ask the ComboBox to refresh its options from its source (that we just updated)
diff --git a/Engine/Plugins/VirtualProduction/CameraCalibration/Source/CameraCalibrationEditor/Private/AssetEditor/SSimulcamViewport.cpp b/Engine/Plugins/VirtualProduction/CameraCalibration/Source/CameraCalibrationEditor/Private/AssetEditor/SSimulcamViewport.cpp
--- a/Engine/Plugins/VirtualProduction/CameraCalibration/Source/CameraCalibrationEditor/Private/AssetEditor/SSimulcamViewport.cpp
+++ b/Engine/Plugins/VirtualProduction/CameraCalibration/Source/CameraCalibrationEditor/Private/AssetEditor/SSimulcamViewport.cpp
@@ -8,15 +8,16 @@
 void SSimulcamViewport::Construct(const FArguments& InArgs, UTexture* InTexture)
 {
 	OnSimulcamViewportClicked = InArgs._OnSimulcamViewportClicked;
+	OnSimulcamViewportInputKey = InArgs._OnSimulcamViewportInputKey;
 
 	Texture = TStrongObjectPtr<UTexture>(InTexture);
 
 	TextureViewport = SNew(SSimulcamEditorViewport, SharedThis(this), InArgs._WithZoom.Get(), InArgs._WithPan.Get());
 
 	ChildSlot
-	[
-		TextureViewport.ToSharedRef()
-	];
+		[
+			TextureViewport.ToSharedRef()
+		];
 }
 
 UTexture* SSimulcamViewport::GetTexture() const
@@ -28,4 +29,13 @@ bool SSimulcamViewport::HasValidTextureResource() const
 {
 	UTexture* CurrentTexture = GetTexture();
 	return CurrentTexture != nullptr && CurrentTexture->GetResource() != nullptr;
-}
\ No newline at end of file
+}
+
+bool SSimulcamViewport::OnViewportInputKey(const FKey& Key, const EInputEvent& Event)
+{ 
+	if (OnSimulcamViewportInputKey.IsBound())
+	{
+		return OnSimulcamViewportInputKey.Execute(Key, Event);
+	}
+	return false;
+}
diff --git a/Engine/Plugins/VirtualProduction/CameraCalibration/Source/CameraCalibrationEditor/Private/AssetEditor/SSimulcamViewport.h b/Engine/Plugins/VirtualProduction/CameraCalibration/Source/CameraCalibrationEditor/Private/AssetEditor/SSimulcamViewport.h
--- a/Engine/Plugins/VirtualProduction/CameraCalibration/Source/CameraCalibrationEditor/Private/AssetEditor/SSimulcamViewport.h
+++ b/Engine/Plugins/VirtualProduction/CameraCalibration/Source/CameraCalibrationEditor/Private/AssetEditor/SSimulcamViewport.h
@@ -2,14 +2,17 @@
 
 #pragma once
 
-#include "Widgets/DeclarativeSyntaxSupport.h"
 #include "Widgets/SCompoundWidget.h"
+
+#include "Engine/EngineBaseTypes.h"
 #include "UObject/GCObject.h"
 #include "UObject/StrongObjectPtr.h"
+#include "Widgets/DeclarativeSyntaxSupport.h"
 
 class UTexture;
 
 struct FGeometry;
+struct FKey;
 struct FPointerEvent;
 
 class SSimulcamEditorViewport;
@@ -17,8 +20,11 @@ class SSimulcamEditorViewport;
 /** Delegate to be executed when the viewport area is clicked */
 DECLARE_DELEGATE_TwoParams(FSimulcamViewportClickedEventHandler, const FGeometry& MyGeometry, const FPointerEvent& PointerEvent);
 
+/** Delegate to be executed when the viewport area receives a key event */
+DECLARE_DELEGATE_RetVal_TwoParams(bool, FSimulcamViewportInputKeyEventHandler, const FKey& Key, const EInputEvent& Event);
+
 /**
- * UI to display the provided UTexture and respond to mouse clicks. 
+ * UI to display the provided UTexture and respond to user input.
  */
 class SSimulcamViewport : public SCompoundWidget
 {
@@ -27,6 +33,7 @@ public:
 	: _WithZoom(true), _WithPan(true)
 	{}
 	SLATE_EVENT(FSimulcamViewportClickedEventHandler, OnSimulcamViewportClicked)
+	SLATE_EVENT(FSimulcamViewportInputKeyEventHandler, OnSimulcamViewportInputKey)
 	SLATE_ATTRIBUTE(bool, WithZoom)
 	SLATE_ATTRIBUTE(bool, WithPan)
 	SLATE_END_ARGS()
@@ -45,14 +52,20 @@ public:
 	/** Is the camera texture valid ? */
 	bool HasValidTextureResource() const;
 
-	/** called whenever the viewport is clicked */
+	/** Called whenever the viewport is clicked */
 	void OnViewportClicked(const FGeometry& MyGeometry, const FPointerEvent& PointerEvent) { OnSimulcamViewportClicked.ExecuteIfBound(MyGeometry, PointerEvent); }
 
+	/** Called when the viewport receives input key presses*/
+	bool OnViewportInputKey(const FKey& Key, const EInputEvent& Event);
+
 private:
 
 	/** Delegate to be executed when the viewport area is clicked */
 	FSimulcamViewportClickedEventHandler OnSimulcamViewportClicked;
 
+	/** Delegate to be executed when the viewport receives input key presses */
+	FSimulcamViewportInputKeyEventHandler OnSimulcamViewportInputKey;
+
 	TSharedPtr<SSimulcamEditorViewport> TextureViewport;
 
 	TStrongObjectPtr<UTexture> Texture;
diff --git a/Engine/Plugins/VirtualProduction/CameraCalibration/Source/CameraCalibrationEditor/Private/AssetEditor/Viewport/SimulcamEditorViewportClient.cpp b/Engine/Plugins/VirtualProduction/CameraCalibration/Source/CameraCalibrationEditor/Private/AssetEditor/Viewport/SimulcamEditorViewportClient.cpp
--- a/Engine/Plugins/VirtualProduction/CameraCalibration/Source/CameraCalibrationEditor/Private/AssetEditor/Viewport/SimulcamEditorViewportClient.cpp
+++ b/Engine/Plugins/VirtualProduction/CameraCalibration/Source/CameraCalibrationEditor/Private/AssetEditor/Viewport/SimulcamEditorViewportClient.cpp
@@ -185,6 +185,12 @@ bool FSimulcamEditorViewportClient::InputKey(FViewport* Viewport, int32 Controll
 				return true;
 			}
 		}
+
+		return SimulcamViewportWeakPtr.Pin()->OnViewportInputKey(Key, Event);
+	}
+	else if (Event == IE_Repeat)
+	{
+		return SimulcamViewportWeakPtr.Pin()->OnViewportInputKey(Key, Event);
 	}
 
 	return false;
diff --git a/Engine/Plugins/VirtualProduction/CameraCalibration/Source/CameraCalibrationEditor/Private/Calibrators/CameraImageCenterAlgoManual.cpp b/Engine/Plugins/VirtualProduction/CameraCalibration/Source/CameraCalibrationEditor/Private/Calibrators/CameraImageCenterAlgoManual.cpp
new file mode 100644
--- /dev/null
+++ b/Engine/Plugins/VirtualProduction/CameraCalibration/Source/CameraCalibrationEditor/Private/Calibrators/CameraImageCenterAlgoManual.cpp
@@ -0,0 +1,253 @@
+// Copyright Epic Games, Inc. All Rights Reserved.
+
+#include "CameraImageCenterAlgoManual.h"
+
+#include "AssetEditor/CameraCalibrationStepsController.h"
+#include "AssetEditor/ImageCenterTool.h"
+#include "CameraCalibrationSettings.h"
+#include "Misc/MessageDialog.h"
+#include "ScopedTransaction.h"
+#include "UI/CameraCalibrationWidgetHelpers.h"
+
+
+#define LOCTEXT_NAMESPACE "CameraImageCenterAlgoManual"
+
+void UCameraImageCenterAlgoManual::Initialize(UImageCenterTool* InImageCenterTool)
+{
+	Tool = InImageCenterTool;
+
+	// Cache the LensFile
+	if (FCameraCalibrationStepsController* StepsController = Tool->GetCameraCalibrationStepsController())
+	{
+		LensFile = StepsController->GetLensFile();
+	}
+
+	// Initialize sensitivity options for the combobox
+	SensitivtyOptions.Add(MakeShared<float>(0.25f));
+	SensitivtyOptions.Add(MakeShared<float>(0.5f));
+	SensitivtyOptions.Add(MakeShared<float>(1.0f));
+	SensitivtyOptions.Add(MakeShared<float>(2.0f));
+	SensitivtyOptions.Add(MakeShared<float>(5.0f));
+	SensitivtyOptions.Add(MakeShared<float>(10.0f));
+}
+
+void UCameraImageCenterAlgoManual::Shutdown()
+{
+	// Prompt the user to save any unsaved changes before shutting down the algo
+	ApplyOrRevertAdjustedImageCenter();
+
+	Tool.Reset();
+}
+
+void UCameraImageCenterAlgoManual::Activate()
+{
+	FCameraCalibrationStepsController* StepsController = Tool->GetCameraCalibrationStepsController();
+	if (!StepsController)
+	{
+		return;
+	}
+
+	if (!LensFile.IsValid())
+	{
+		return;
+	}
+
+	// Cache the render target size to use when calculating the fine adjustments
+	UTextureRenderTarget2D* RenderTarget = StepsController->GetRenderTarget();
+	if (!RenderTarget)
+	{
+		return;
+	}
+
+	RenderTargetSize = FIntPoint(RenderTarget->SizeX, RenderTarget->SizeY);
+
+	const FLensFileEvalData* EvalData = StepsController->GetLensFileEvalData();
+	if (!EvalData)
+	{
+		return;
+	}
+
+	// Cache the original image center point for the current focus and zoom
+	CurrentEvalFocus = EvalData->Input.Focus;
+	CurrentEvalZoom = EvalData->Input.Zoom;
+	LensFile->EvaluateImageCenterParameters(CurrentEvalFocus, CurrentEvalZoom, OriginalImageCenter);
+	AdjustedImageCenter = OriginalImageCenter;
+
+	bIsActive = true;
+}
+
+void UCameraImageCenterAlgoManual::Deactivate()
+{
+	// Prompt the user to save any unsaved changes before deactivating the algo
+	ApplyOrRevertAdjustedImageCenter();
+
+	bIsActive = false;
+}
+
+void UCameraImageCenterAlgoManual::Tick(float DeltaTime)
+{
+	FCameraCalibrationStepsController* StepsController = Tool->GetCameraCalibrationStepsController();
+	if (!StepsController)
+	{
+		return;
+	}
+
+	if (LensFile.IsValid())
+	{
+		if (const FLensFileEvalData* EvalData = StepsController->GetLensFileEvalData())
+		{
+			// Check if the focus or zoom has changed since the last tick
+			const float InputTolerance = GetDefault<UCameraCalibrationSettings>()->GetCalibrationInputTolerance();
+			if (!FMath::IsNearlyEqual(CurrentEvalFocus, EvalData->Input.Focus, InputTolerance) || !FMath::IsNearlyEqual(CurrentEvalZoom, EvalData->Input.Zoom, InputTolerance))
+			{
+				// Prompt the user to save any unsaved changes before caching a new original image center based on the new focus/zoom
+				ApplyOrRevertAdjustedImageCenter();
+
+				CurrentEvalFocus = EvalData->Input.Focus;
+				CurrentEvalZoom = EvalData->Input.Zoom;
+				LensFile->EvaluateImageCenterParameters(CurrentEvalFocus, CurrentEvalZoom, OriginalImageCenter);
+				AdjustedImageCenter = OriginalImageCenter;
+			}
+		}
+	}
+}
+
+bool UCameraImageCenterAlgoManual::OnViewportInputKey(const FKey& InKey, const EInputEvent& InEvent)
+{
+	if (!LensFile.IsValid())
+	{
+		return false;
+	}
+
+	// Guard against divide-by-zero
+	if ((RenderTargetSize.X == 0) || (RenderTargetSize.Y == 0))
+	{
+		return false;
+	}
+
+	bool bHandled = false;
+
+	// Arrow Key handling
+	if (InKey == EKeys::Left)
+	{
+		AdjustedImageCenter.PrincipalPoint.X -= (AdjustmentIncrement / RenderTargetSize.X);
+		bHandled = true;
+	}
+	else if (InKey == EKeys::Right)
+	{
+		AdjustedImageCenter.PrincipalPoint.X += (AdjustmentIncrement / RenderTargetSize.X);
+		bHandled = true;
+	}
+	else if (InKey == EKeys::Up)
+	{
+		AdjustedImageCenter.PrincipalPoint.Y -= (AdjustmentIncrement / RenderTargetSize.Y);
+		bHandled = true;
+	}
+	else if (InKey == EKeys::Down)
+	{
+		AdjustedImageCenter.PrincipalPoint.Y += (AdjustmentIncrement / RenderTargetSize.Y);
+		bHandled = true;
+	}
+
+	if (bHandled)
+	{
+		// First, attempt to modify an existing image point in the image center table
+		bool bPointExistsInTable = LensFile->ImageCenterTable.SetPoint(CurrentEvalFocus, CurrentEvalZoom, AdjustedImageCenter);
+
+		// If no point exists at the specified focus and zoom, add a new image center point
+		if (!bPointExistsInTable)
+		{
+			LensFile->AddImageCenterPoint(CurrentEvalFocus, CurrentEvalZoom, AdjustedImageCenter);
+		}
+	}
+
+	return bHandled;
+}
+
+TSharedRef<SWidget> UCameraImageCenterAlgoManual::BuildUI()
+{
+	return SNew(SVerticalBox)
+
+		+ SVerticalBox::Slot() // Sensitivity
+		.VAlign(EVerticalAlignment::VAlign_Top)
+		.AutoHeight()
+		.MaxHeight(FCameraCalibrationWidgetHelpers::DefaultRowHeight)
+		[FCameraCalibrationWidgetHelpers::BuildLabelWidgetPair(LOCTEXT("SensitivityWidget", "Sensitivity"), BuildSensitivityWidget())]
+	;
+}
+
+bool UCameraImageCenterAlgoManual::HasImageCenterChanged()
+{
+	return (OriginalImageCenter.PrincipalPoint != AdjustedImageCenter.PrincipalPoint);
+}
+
+void UCameraImageCenterAlgoManual::OnSavedImageCenter()
+{
+	// Before beginning a transaction, restore the lens file image center point to its original state
+	const FImageCenterInfo CachedAdjustedImageCenter = AdjustedImageCenter;
+	AdjustedImageCenter = OriginalImageCenter;
+	LensFile->ImageCenterTable.SetPoint(CurrentEvalFocus, CurrentEvalZoom, OriginalImageCenter);
+
+	{
+		FScopedTransaction Transaction(LOCTEXT("SaveAdjustedImageCenter", "Save Adjusted Image Center"));
+		this->Modify();
+		LensFile->Modify();
+
+		// After beginning the transaction, reapply the final modified image center
+		AdjustedImageCenter = CachedAdjustedImageCenter;
+		LensFile->ImageCenterTable.SetPoint(CurrentEvalFocus, CurrentEvalZoom, AdjustedImageCenter);
+
+		// Reset the original image center to the adjusted point to indicate that the change has been written
+		OriginalImageCenter = AdjustedImageCenter;
+	}
+}
+
+TSharedRef<SWidget> UCameraImageCenterAlgoManual::BuildSensitivityWidget()
+{
+	return SNew(SComboBox<TSharedPtr<float>>)
+		.OptionsSource(&SensitivtyOptions)
+		.OnSelectionChanged_Lambda([&](TSharedPtr<float> NewValue, ESelectInfo::Type Type) -> void
+		{
+			AdjustmentIncrement = *NewValue;
+		})
+		.OnGenerateWidget_Lambda([&](TSharedPtr<float> InOption) -> TSharedRef<SWidget>
+		{
+			return SNew(STextBlock).Text(FText::AsNumber(*InOption));
+		})
+		.InitiallySelectedItem(SensitivtyOptions[2])
+		[
+			SNew(STextBlock)
+			.Text_Lambda([&]() -> FText
+			{
+				return FText::AsNumber(AdjustmentIncrement);
+			})
+		];
+}
+
+void UCameraImageCenterAlgoManual::ApplyOrRevertAdjustedImageCenter()
+{
+	// Warn the user that they have modified the image center, and ask if they want to save or revert their changes
+	if (OriginalImageCenter.PrincipalPoint != AdjustedImageCenter.PrincipalPoint)
+	{
+		FText Message = LOCTEXT("UnsavedImageCenterChanges",
+			"Would you like to apply the changes of the modified image center to the lens file?");
+
+		EAppReturnType::Type ReturnValue = FMessageDialog::Open(EAppMsgType::YesNo, Message);
+
+		if (ReturnValue == EAppReturnType::No)
+		{
+			// Reset the image center back to the original value
+			if (LensFile.IsValid())
+			{
+				LensFile->ImageCenterTable.SetPoint(CurrentEvalFocus, CurrentEvalZoom, OriginalImageCenter);
+				AdjustedImageCenter = OriginalImageCenter;
+			}
+		}
+		else if (ReturnValue == EAppReturnType::Yes)
+		{
+			OnSavedImageCenter();
+		}
+	}
+}
+
+#undef LOCTEXT_NAMESPACE
diff --git a/Engine/Plugins/VirtualProduction/CameraCalibration/Source/CameraCalibrationEditor/Private/Calibrators/CameraImageCenterAlgoManual.h b/Engine/Plugins/VirtualProduction/CameraCalibration/Source/CameraCalibrationEditor/Private/Calibrators/CameraImageCenterAlgoManual.h
new file mode 100644
--- /dev/null
+++ b/Engine/Plugins/VirtualProduction/CameraCalibration/Source/CameraCalibrationEditor/Private/Calibrators/CameraImageCenterAlgoManual.h
@@ -0,0 +1,73 @@
+// Copyright Epic Games, Inc. All Rights Reserved.
+
+#pragma once
+
+#include "CameraImageCenterAlgo.h"
+
+#include "LensData.h"
+#include "LensFile.h"
+
+#include "CameraImageCenterAlgoManual.generated.h"
+
+class SWidget;
+
+/**
+ * Implements an image center adjustment algorithm.
+ * The Manual algorithm simply takes manual input from the user and adjusts the image center by small increments
+ */
+UCLASS()
+class UCameraImageCenterAlgoManual : public UCameraImageCenterAlgo
+{
+	GENERATED_BODY()
+
+public:
+
+	//~ Begin UCameraImageCenterAlgo
+	virtual void Initialize(UImageCenterTool* InImageCenterTool) override;
+	virtual void Shutdown() override;
+	virtual void Activate() override;
+	virtual void Deactivate() override;
+	virtual void Tick(float DeltaTime) override;
+	virtual bool OnViewportInputKey(const FKey& InKey, const EInputEvent& InEvent) override;
+	virtual TSharedRef<SWidget> BuildUI() override;
+	virtual FName FriendlyName() const override { return TEXT("Manual Image Center Adjustment"); };
+	virtual bool HasImageCenterChanged() override;
+	virtual void OnSavedImageCenter() override;
+	//~ End UCameraImageCenterAlgo
+
+private:
+	/** Builds the UI of the sensitivity widget */
+	TSharedRef<SWidget> BuildSensitivityWidget();
+
+	/** Open a dialog to let the user decide whether to save their adjustments or revert to the original state */
+	void ApplyOrRevertAdjustedImageCenter();
+
+private:
+	/** Cached original image center that was evaluated at the current focus and zoom */
+	UPROPERTY(Transient)
+	FImageCenterInfo OriginalImageCenter;
+
+	/** Latest image center that may have been adjusted by manual user input */
+	UPROPERTY(Transient)
+	FImageCenterInfo AdjustedImageCenter;
+
+private:
+	/** The image center tool controller */
+	TWeakObjectPtr<UImageCenterTool> Tool;
+
+	/** The LensFile being edited */
+	TWeakObjectPtr<ULensFile> LensFile;
+
+	/** Cached focus and zoom that were used to evaluate the original image center */
+	float CurrentEvalFocus;
+	float CurrentEvalZoom;
+
+	/** Size of the render target displayed in the viewport */
+	FIntPoint RenderTargetSize;
+
+	/** Increment (in pixels) that the image center will be adjusted by each individual key press */
+	float AdjustmentIncrement = 0.5f;
+
+	/** Array of sensitivity options to display to the user to set the adjustement increment */
+	TArray<TSharedPtr<float>> SensitivtyOptions;
+};
diff --git a/Engine/Plugins/VirtualProduction/CameraCalibration/Source/CameraCalibrationEditor/Private/Calibrators/CameraLensDistortionAlgoCheckerboard.cpp b/Engine/Plugins/VirtualProduction/CameraCalibration/Source/CameraCalibrationEditor/Private/Calibrators/CameraLensDistortionAlgoCheckerboard.cpp
--- a/Engine/Plugins/VirtualProduction/CameraCalibration/Source/CameraCalibrationEditor/Private/Calibrators/CameraLensDistortionAlgoCheckerboard.cpp
+++ b/Engine/Plugins/VirtualProduction/CameraCalibration/Source/CameraCalibrationEditor/Private/Calibrators/CameraLensDistortionAlgoCheckerboard.cpp
@@ -5,34 +5,15 @@
 #include "AssetEditor/CameraCalibrationStepsController.h"
 #include "AssetEditor/LensDistortionTool.h"
 #include "AssetEditor/SSimulcamViewport.h"
-#include "Camera/CameraActor.h"
-#include "Camera/CameraComponent.h"
 #include "CameraCalibrationCheckerboard.h"
-#include "CameraCalibrationEditorLog.h"
-#include "CameraCalibrationUtils.h"
-#include "Editor.h"
 #include "EditorFontGlyphs.h"
-#include "Engine/TextureRenderTarget2D.h"
 #include "EngineUtils.h"
-#include "GameFramework/Actor.h"
-#include "GenericPlatform/GenericApplication.h"
-#include "Input/Events.h"
-#include "Internationalization/Text.h"
-#include "Layout/Geometry.h"
-#include "LensFile.h"
-#include "Math/Color.h"
-#include "Math/UnrealMathUtility.h"
-#include "Math/Vector.h"
 #include "Misc/MessageDialog.h"
 #include "PropertyCustomizationHelpers.h"
 #include "SphericalLensDistortionModelHandler.h"
 #include "UI/CameraCalibrationWidgetHelpers.h"
 #include "Widgets/Input/SButton.h"
 #include "Widgets/Input/SCheckBox.h"
-#include "Widgets/Layout/SSpacer.h"
-#include "Widgets/Views/SListView.h"
-#include "Widgets/SNullWidget.h"
-#include "Widgets/SWidget.h"
 
 
 #if WITH_OPENCV
diff --git a/Engine/Plugins/VirtualProduction/CameraCalibrationCore/Source/CameraCalibrationCore/Private/CameraCalibrationSubsystem.cpp b/Engine/Plugins/VirtualProduction/CameraCalibrationCore/Source/CameraCalibrationCore/Private/CameraCalibrationSubsystem.cpp
--- a/Engine/Plugins/VirtualProduction/CameraCalibrationCore/Source/CameraCalibrationCore/Private/CameraCalibrationSubsystem.cpp
+++ b/Engine/Plugins/VirtualProduction/CameraCalibrationCore/Source/CameraCalibrationCore/Private/CameraCalibrationSubsystem.cpp
@@ -218,6 +218,22 @@ TArray<FName> UCameraCalibrationSubsystem::GetCameraNodalOffsetAlgos() const
 	return OutKeys;
 }
 
+TSubclassOf<UCameraImageCenterAlgo> UCameraCalibrationSubsystem::GetCameraImageCenterAlgo(FName Name) const
+{
+	if (CameraImageCenterAlgosMap.Contains(Name))
+	{
+		return CameraImageCenterAlgosMap[Name];
+	}
+	return nullptr;
+}
+
+TArray<FName> UCameraCalibrationSubsystem::GetCameraImageCenterAlgos() const
+{
+	TArray<FName> OutKeys;
+	CameraImageCenterAlgosMap.GetKeys(OutKeys);
+	return OutKeys;
+}
+
 TSubclassOf<UCameraCalibrationStep> UCameraCalibrationSubsystem::GetCameraCalibrationStep(FName Name) const
 {
 	if (CameraCalibrationStepsMap.Contains(Name))
@@ -240,6 +256,20 @@ void UCameraCalibrationSubsystem::Initialize(FSubsystemCollectionBase& Collectio
 
 	PostEngineInitHandle = FCoreDelegates::OnPostEngineInit.AddLambda([&]()
 	{
+		// Find Image Center Algos
+		{
+			TArray<TSubclassOf<UCameraImageCenterAlgo>> Algos;
+
+			for (TObjectIterator<UClass> AlgoIt; AlgoIt; ++AlgoIt)
+			{
+				if (AlgoIt->IsChildOf(UCameraImageCenterAlgo::StaticClass()) && !AlgoIt->HasAnyClassFlags(CLASS_Abstract | CLASS_Deprecated))
+				{
+					const UCameraImageCenterAlgo* Algo = CastChecked<UCameraImageCenterAlgo>(AlgoIt->GetDefaultObject());
+					CameraImageCenterAlgosMap.Add(Algo->FriendlyName(), TSubclassOf<UCameraImageCenterAlgo>(*AlgoIt));
+				}
+			}
+		}
+
 		// Find Nodal Offset Algos
 		{
 			TArray<TSubclassOf<UCameraNodalOffsetAlgo>> Algos;
@@ -272,6 +302,7 @@ void UCameraCalibrationSubsystem::Initialize(FSubsystemCollectionBase& Collectio
 void UCameraCalibrationSubsystem::Deinitialize()
 {
 	LensModelMap.Empty(0);
+	CameraImageCenterAlgosMap.Empty(0);
 	CameraNodalOffsetAlgosMap.Empty(0);
 	CameraCalibrationStepsMap.Empty(0);
 
diff --git a/Engine/Plugins/VirtualProduction/CameraCalibrationCore/Source/CameraCalibrationCore/Public/CameraCalibrationStep.h b/Engine/Plugins/VirtualProduction/CameraCalibrationCore/Source/CameraCalibrationCore/Public/CameraCalibrationStep.h
--- a/Engine/Plugins/VirtualProduction/CameraCalibrationCore/Source/CameraCalibrationCore/Public/CameraCalibrationStep.h
+++ b/Engine/Plugins/VirtualProduction/CameraCalibrationCore/Source/CameraCalibrationCore/Public/CameraCalibrationStep.h
@@ -3,12 +3,14 @@
 #pragma once
 
 #include "CoreMinimal.h"
+#include "Engine/EngineBaseTypes.h"
 #include "Templates/SharedPointer.h"
 #include "Widgets/SBoxPanel.h"
 
 #include "CameraCalibrationStep.generated.h"
 
 struct FGeometry;
+struct FKey;
 struct FPointerEvent;
 
 class FCameraCalibrationStepsController;
@@ -35,6 +37,9 @@ public:
 	/** Callback when viewport is clicked. Returns false if the event was not handled. */
 	virtual bool OnViewportClicked(const FGeometry& MyGeometry, const FPointerEvent& MouseEvent) { return false;  };
 
+	/** Callback when key is pressed while viewport is focused. Returns false if the event was not handled. */
+	virtual bool OnViewportInputKey(const FKey& InKey, const EInputEvent& InEvent) { return false; };
+
 	/** Returns the UI of this camera calibration step. Expected to only be called once */
 	virtual TSharedRef<SWidget> BuildUI() { return SNew(SVerticalBox); };
 
diff --git a/Engine/Plugins/VirtualProduction/CameraCalibrationCore/Source/CameraCalibrationCore/Public/CameraCalibrationSubsystem.h b/Engine/Plugins/VirtualProduction/CameraCalibrationCore/Source/CameraCalibrationCore/Public/CameraCalibrationSubsystem.h
--- a/Engine/Plugins/VirtualProduction/CameraCalibrationCore/Source/CameraCalibrationCore/Public/CameraCalibrationSubsystem.h
+++ b/Engine/Plugins/VirtualProduction/CameraCalibrationCore/Source/CameraCalibrationCore/Public/CameraCalibrationSubsystem.h
@@ -7,6 +7,7 @@
 #include "CoreMinimal.h"
 
 #include "CameraCalibrationTypes.h"
+#include "CameraImageCenterAlgo.h"
 #include "CameraNodalOffsetAlgo.h"
 #include "CameraCalibrationStep.h"
 #include "Containers/ArrayView.h"
@@ -79,6 +80,14 @@ public:
 	UFUNCTION(BlueprintCallable, Category = "Lens Distortion")
 	TArray<FName> GetCameraNodalOffsetAlgos() const;
 
+	/** Returns the image center algorithm by name */
+	UFUNCTION(BlueprintCallable, Category = "Lens Distortion")
+	TSubclassOf<UCameraImageCenterAlgo> GetCameraImageCenterAlgo(FName Name) const;
+
+	/** Returns an array with the names of the available image center algorithms */
+	UFUNCTION(BlueprintCallable, Category = "Lens Distortion")
+	TArray<FName> GetCameraImageCenterAlgos() const;
+
 	/** Returns an array with the names of the available camera calibration steps */
 	UFUNCTION(BlueprintCallable, Category = "Lens Distortion")
 	TArray<FName> GetCameraCalibrationSteps() const;
@@ -124,6 +133,10 @@ private:
 	UPROPERTY(Transient)
 	TMap<FName, TSubclassOf<UCameraNodalOffsetAlgo>> CameraNodalOffsetAlgosMap;
 
+	/** Holds the registered camera image center algos */
+	UPROPERTY(Transient)
+	TMap<FName, TSubclassOf<UCameraImageCenterAlgo>> CameraImageCenterAlgosMap;
+
 	/** Holds the registered camera calibration steps */
 	UPROPERTY(Transient)
 	TMap<FName, TSubclassOf<UCameraCalibrationStep>> CameraCalibrationStepsMap;
diff --git a/Engine/Plugins/VirtualProduction/CameraCalibrationCore/Source/CameraCalibrationCore/Public/CameraImageCenterAlgo.h b/Engine/Plugins/VirtualProduction/CameraCalibrationCore/Source/CameraCalibrationCore/Public/CameraImageCenterAlgo.h
new file mode 100644
--- /dev/null
+++ b/Engine/Plugins/VirtualProduction/CameraCalibrationCore/Source/CameraCalibrationCore/Public/CameraImageCenterAlgo.h
@@ -0,0 +1,73 @@
+// Copyright Epic Games, Inc. All Rights Reserved.
+
+#pragma once
+
+#include "CoreMinimal.h"
+#include "Engine/EngineBaseTypes.h"
+#include "Widgets/SNullWidget.h"
+#include "Widgets/Text/STextBlock.h"
+
+#include "CameraImageCenterAlgo.generated.h"
+
+struct FGeometry;
+struct FKey;
+struct FPointerEvent;
+
+class UImageCenterTool;
+class SWidget;
+
+/**
+ * UCameraImageCenterAlgo defines the interface that any image center algorithm should implement
+ * in order to be used and listed by the Image Center Tool.
+ */
+UCLASS(Abstract)
+class CAMERACALIBRATIONCORE_API UCameraImageCenterAlgo : public UObject
+{
+	GENERATED_BODY()
+
+public:
+
+	/** Make sure you initialize before using the object */
+	virtual void Initialize(UImageCenterTool* InImageCenterTool) {};
+
+	/** Clean up resources and don't use ImageCenterTool anymore */
+	virtual void Shutdown() {};
+
+	/** Perform any required set when this algo becomes the active one */
+	virtual void Activate() {};
+
+	/** Perform any required cleanup when this algo is no longer active */
+	virtual void Deactivate() {};
+
+	/** Returns true if this algo is active */
+	virtual bool IsActive() const { return bIsActive; }
+
+	/** Called every frame */
+	virtual void Tick(float DeltaTime) {};
+
+	/** Callback when viewport is clicked. Returns false if the event was not handled. */
+	virtual bool OnViewportClicked(const FGeometry& MyGeometry, const FPointerEvent& MouseEvent) { return false; };
+
+	/** Callback when viewport receives input key presses. Returns false if the event was not handled. */
+	virtual bool OnViewportInputKey(const FKey& InKey, const EInputEvent& InEvent) { return false; };
+
+	/** Returns the UI of this calibrator. Expected to only be called once */
+	virtual TSharedRef<SWidget> BuildUI() { return SNullWidget::NullWidget; };
+
+	/** Returns true if the algo has changed the image center for the current focus/zoom */
+	virtual bool HasImageCenterChanged() { return false; };
+
+	/** Returns a descriptive name/title of this image center algorithm */
+	virtual FName FriendlyName() const { return TEXT("Invalid Name"); };
+
+	/** Called when the current offset was saved */
+	virtual void OnSavedImageCenter() { };
+
+	/** Called to present the user with instructions on how to use this algo */
+	virtual TSharedRef<SWidget> BuildHelpWidget() { return SNew(STextBlock).Text(FText::FromString(TEXT("Coming soon!"))); };
+
+protected:
+
+	/** Whether this algo is current active or not */
+	bool bIsActive = false;
+};
diff --git a/Engine/Plugins/VirtualProduction/CameraCalibrationCore/Source/CameraCalibrationCore/Public/CameraLensDistortionAlgo.h b/Engine/Plugins/VirtualProduction/CameraCalibrationCore/Source/CameraCalibrationCore/Public/CameraLensDistortionAlgo.h
--- a/Engine/Plugins/VirtualProduction/CameraCalibrationCore/Source/CameraCalibrationCore/Public/CameraLensDistortionAlgo.h
+++ b/Engine/Plugins/VirtualProduction/CameraCalibrationCore/Source/CameraCalibrationCore/Public/CameraLensDistortionAlgo.h
@@ -65,6 +65,6 @@ public:
 		return false; 
 	};
 
-	/** Called to present the user with instructions on how to this this algo */
+	/** Called to present the user with instructions on how to use this algo */
 	virtual TSharedRef<SWidget> BuildHelpWidget() { return SNew(STextBlock).Text(FText::FromString(TEXT("Coming soon!"))); };
 };
diff --git a/Engine/Plugins/VirtualProduction/CameraCalibrationCore/Source/CameraCalibrationCore/Public/CameraNodalOffsetAlgo.h b/Engine/Plugins/VirtualProduction/CameraCalibrationCore/Source/CameraCalibrationCore/Public/CameraNodalOffsetAlgo.h
--- a/Engine/Plugins/VirtualProduction/CameraCalibrationCore/Source/CameraCalibrationCore/Public/CameraNodalOffsetAlgo.h
+++ b/Engine/Plugins/VirtualProduction/CameraCalibrationCore/Source/CameraCalibrationCore/Public/CameraNodalOffsetAlgo.h
@@ -50,6 +50,6 @@ public:
 	/** Called when the current offset was saved */
 	virtual void OnSavedNodalOffset() { };
 
-	/** Called to present the user with instructions on how to this this algo */
+	/** Called to present the user with instructions on how to use this algo */
 	virtual TSharedRef<SWidget> BuildHelpWidget() { return SNew(STextBlock).Text(FText::FromString(TEXT("Coming soon!")));  };
 };
