[CameraCalibration] - Add overlay support to lens calibration tool. Tool now displays a picker to select from a list of registered overlay materials, as well as its material parameters to support overlay configuration directly in the tool.
[CameraCalibration] - Add an overlay to lens distortion checkboard algorithm that progressively displays the coverage achieved when taking images of a checkerboard for distortion calibration

diff --git a/Engine/Plugins/Compositing/Composure/Source/Composure/Private/CompositingElements/CompositingMaterialPass.cpp b/Engine/Plugins/Compositing/Composure/Source/Composure/Private/CompositingElements/CompositingMaterialPass.cpp
--- a/Engine/Plugins/Compositing/Composure/Source/Composure/Private/CompositingElements/CompositingMaterialPass.cpp
+++ b/Engine/Plugins/Compositing/Composure/Source/Composure/Private/CompositingElements/CompositingMaterialPass.cpp
@@ -45,6 +45,11 @@ bool FCompositingMaterial::ApplyParamOverrides(const ICompositingTextureLookupTa
 				MatInstance->SetVectorParameterValue(VecOverride.Key, VecOverride.Value);
 			}
 
+			for (TPair<FName, TObjectPtr<UTexture>>& TextureOverride : TextureParamOverrides)
+			{
+				MatInstance->SetTextureParameterValue(TextureOverride.Key, TextureOverride.Value.Get());
+			}
+
 			bParamsModified = false;
 		}
 
@@ -218,6 +223,40 @@ void FCompositingMaterial::ResetVectorOverride(const FName ParamName)
 	}
 }
 
+void FCompositingMaterial::SetTextureOverride(const FName ParamName, UTexture* ParamVal)
+{
+	TextureParamOverrides.Add(ParamName, ParamVal);
+	if (CachedMID)
+	{
+		CachedMID->SetTextureParameterValue(ParamName, ParamVal);
+	}
+	else
+	{
+		bParamsModified = true;
+	}
+}
+
+bool FCompositingMaterial::GetTextureOverride(const FName ParamName, UTexture*& OutParamVal)
+{
+	if (TObjectPtr<UTexture>* ExistingOverride = TextureParamOverrides.Find(ParamName))
+	{
+		OutParamVal = *ExistingOverride;
+		return true;
+	}
+	return false;
+}
+
+void FCompositingMaterial::ResetTextureOverride(const FName ParamName)
+{
+	TextureParamOverrides.Remove(ParamName);
+
+	UTexture* DefaultVal;
+	if (CachedMID && CachedMID->GetTextureParameterDefaultValue(ParamName, DefaultVal))
+	{
+		CachedMID->SetTextureParameterValue(ParamName, DefaultVal);
+	}
+}
+
 void FCompositingMaterial::ResetAllParamOverrides()
 {
 	while (ScalarParamOverrides.Num() > 0)
diff --git a/Engine/Plugins/Compositing/Composure/Source/Composure/Public/CompositingElements/CompositingMaterialPass.h b/Engine/Plugins/Compositing/Composure/Source/Composure/Public/CompositingElements/CompositingMaterialPass.h
--- a/Engine/Plugins/Compositing/Composure/Source/Composure/Public/CompositingElements/CompositingMaterialPass.h
+++ b/Engine/Plugins/Compositing/Composure/Source/Composure/Public/CompositingElements/CompositingMaterialPass.h
@@ -3,7 +3,9 @@
 #pragma once
 
 #include "CoreMinimal.h"
+#include "Engine/Texture.h"
 #include "UObject/ObjectMacros.h"
+
 #include "CompositingMaterialPass.generated.h"
 
 class UMaterialInterface;
@@ -18,11 +20,17 @@ struct COMPOSURE_API FCompositingParamPayload
 {
 	GENERATED_USTRUCT_BODY();
 
+	/** Map of scalar material parameter names to their corresponding override values */
 	UPROPERTY()
 	TMap<FName, float> ScalarParamOverrides;
 
+	/** Map of vector material parameter names to their corresponding override values */
 	UPROPERTY()
 	TMap<FName, FLinearColor> VectorParamOverrides;
+
+	/** Map of texture material parameter names to their corresponding override values */
+	UPROPERTY()
+	TMap<FName, TObjectPtr<UTexture>> TextureParamOverrides;
 };
 
 /* FNamedCompMaterialParam
@@ -121,6 +129,15 @@ public:
 	bool GetVectorOverride(const FName ParamName, FLinearColor& OutParamVal);
 	void ResetVectorOverride(const FName ParamName);
 
+	/** Set the override UTexture value for the named input material parameter */
+	void SetTextureOverride(const FName ParamName, UTexture* ParamVal);
+
+	/** Get the override UTexture value for the named input material parameter */
+	bool GetTextureOverride(const FName ParamName, UTexture*& OutParamVal);
+
+	/** Remove the override for the named input material parameter and reset its value to the default */
+	void ResetTextureOverride(const FName ParamName);
+
 	void ResetAllParamOverrides();
 
 	void MarkDirty() { bParamsModified = true; };
diff --git a/Engine/Plugins/Compositing/OpenCVLensDistortion/Source/OpenCVHelper/Private/OpenCVHelper.cpp b/Engine/Plugins/Compositing/OpenCVLensDistortion/Source/OpenCVHelper/Private/OpenCVHelper.cpp
--- a/Engine/Plugins/Compositing/OpenCVLensDistortion/Source/OpenCVHelper/Private/OpenCVHelper.cpp
+++ b/Engine/Plugins/Compositing/OpenCVLensDistortion/Source/OpenCVHelper/Private/OpenCVHelper.cpp
@@ -134,4 +134,60 @@ UTexture2D* FOpenCVHelper::TextureFromCvMat(cv::Mat& Mat, const FString* Package
 	return Texture;
 }
 
+UTexture2D* FOpenCVHelper::TextureFromCvMat(cv::Mat& Mat, UTexture2D* InTexture)
+{
+	if (!InTexture)
+	{
+		return TextureFromCvMat(Mat);
+	}
+
+	if ((Mat.cols <= 0) || (Mat.rows <= 0))
+	{
+		return nullptr;
+	}
+
+	// Currently we only support G8 and BGRA8
+
+	if (Mat.depth() != CV_8U)
+	{
+		return nullptr;
+	}
+
+	EPixelFormat PixelFormat;
+
+	switch (Mat.channels())
+	{
+	case 1:
+		PixelFormat = PF_G8;
+		break;
+
+	case 4:
+		PixelFormat = PF_B8G8R8A8;
+		break;
+
+	default:
+		return nullptr;
+	}
+
+	if ((InTexture->GetSizeX() != Mat.cols) || (InTexture->GetSizeY() != Mat.rows) || (InTexture->GetPixelFormat() != PixelFormat))
+	{
+		return nullptr;
+	}
+
+	// Copy the pixels from the OpenCV Mat to the Texture
+
+	FTexture2DMipMap& Mip0 = InTexture->PlatformData->Mips[0];
+	void* TextureData = Mip0.BulkData.Lock(LOCK_READ_WRITE);
+
+	const int32 PixelStride = Mat.channels();
+	FMemory::Memcpy(TextureData, Mat.data, SIZE_T(Mat.cols * Mat.rows * PixelStride));
+
+	Mip0.BulkData.Unlock();
+
+	InTexture->UpdateResource();
+
+	return InTexture;
+}
+
+
 #endif // WITH_OPENCV
diff --git a/Engine/Plugins/Compositing/OpenCVLensDistortion/Source/OpenCVHelper/Public/OpenCVHelper.h b/Engine/Plugins/Compositing/OpenCVLensDistortion/Source/OpenCVHelper/Public/OpenCVHelper.h
--- a/Engine/Plugins/Compositing/OpenCVLensDistortion/Source/OpenCVHelper/Public/OpenCVHelper.h
+++ b/Engine/Plugins/Compositing/OpenCVLensDistortion/Source/OpenCVHelper/Public/OpenCVHelper.h
@@ -59,6 +59,7 @@ public:
 	 * @return Texture created out of the given OpenCV Mat.
 	 */
 	static UTexture2D* TextureFromCvMat(cv::Mat& Mat, const FString* PackagePath = nullptr, const FName* TextureName = nullptr);
+	static UTexture2D* TextureFromCvMat(cv::Mat& Mat, UTexture2D* InTexture);
 };
 
 #endif //WITH_OPENCV
diff --git a/Engine/Plugins/VirtualProduction/CameraCalibration/Source/CameraCalibrationEditor/CameraCalibrationEditor.Build.cs b/Engine/Plugins/VirtualProduction/CameraCalibration/Source/CameraCalibrationEditor/CameraCalibrationEditor.Build.cs
--- a/Engine/Plugins/VirtualProduction/CameraCalibration/Source/CameraCalibrationEditor/CameraCalibrationEditor.Build.cs
+++ b/Engine/Plugins/VirtualProduction/CameraCalibration/Source/CameraCalibrationEditor/CameraCalibrationEditor.Build.cs
@@ -10,6 +10,7 @@ namespace UnrealBuildTool.Rules
 				new string[]
 				{
 					"ApplicationCore",
+					"AppFramework",
 					"AssetRegistry",
 					"AssetTools",
 					"CameraCalibration",
diff --git a/Engine/Plugins/VirtualProduction/CameraCalibration/Source/CameraCalibrationEditor/Private/AssetEditor/CameraCalibrationStepsController.cpp b/Engine/Plugins/VirtualProduction/CameraCalibration/Source/CameraCalibrationEditor/Private/AssetEditor/CameraCalibrationStepsController.cpp
--- a/Engine/Plugins/VirtualProduction/CameraCalibration/Source/CameraCalibrationEditor/Private/AssetEditor/CameraCalibrationStepsController.cpp
+++ b/Engine/Plugins/VirtualProduction/CameraCalibration/Source/CameraCalibrationEditor/Private/AssetEditor/CameraCalibrationStepsController.cpp
@@ -108,6 +108,7 @@ namespace CameraCalibrationStepsController
 
 FCameraCalibrationStepsController::FCameraCalibrationStepsController(TWeakPtr<FCameraCalibrationToolkit> InCameraCalibrationToolkit, ULensFile* InLensFile)
 	: CameraCalibrationToolkit(InCameraCalibrationToolkit)
+	, RenderTargetSize(FIntPoint(1920, 1080))
 	, LensFile(TWeakObjectPtr<ULensFile>(InLensFile))
 {
 	check(CameraCalibrationToolkit.IsValid());
@@ -330,6 +331,11 @@ UTextureRenderTarget2D* FCameraCalibrationStepsController::GetRenderTarget() con
 	return RenderTarget.Get();
 }
 
+FIntPoint FCameraCalibrationStepsController::GetCompRenderTargetSize() const
+{
+	return RenderTargetSize;
+}
+
 void FCameraCalibrationStepsController::CreateComp()
 {
 	// Don't do anything if we already created it.
@@ -554,6 +560,21 @@ void FCameraCalibrationStepsController::CreateComp()
 	MaterialPass->SetParameterMapping(TEXT("CG"), *CGLayer->GetActorLabel());
 	MaterialPass->SetParameterMapping(TEXT("MediaPlate"), *MediaPlate->GetActorLabel());
 
+	// Create new overlay transform pass
+	OverlayPass = CastChecked<UCompositingElementMaterialPass>(
+		Comp->CreateNewTransformPass(TEXT("Overlay"), UCompositingElementMaterialPass::StaticClass())
+		);
+
+	if (!OverlayPass.IsValid())
+	{
+		UE_LOG(LogCameraCalibrationEditor, Warning, TEXT("Failed to create 'Overlay' UCompositingElementMaterialPass"));
+		Cleanup();
+		return;
+	}
+
+	// The overlay pass should be disabled until an overlay material is set
+	OverlayPass->SetPassEnabled(false);
+
 	URenderTargetCompositingOutput* RTOutput = Cast<URenderTargetCompositingOutput>(Comp->CreateNewOutputPass(
 		TEXT("SimulcamCalOutput"),
 		URenderTargetCompositingOutput::StaticClass())
@@ -582,7 +603,7 @@ void FCameraCalibrationStepsController::CreateComp()
 	RenderTarget->RenderTargetFormat = RTF_RGBA16f;
 	RenderTarget->ClearColor = FLinearColor::Black;
 	RenderTarget->bAutoGenerateMips = true;
-	RenderTarget->InitAutoFormat(1920, 1080);
+	RenderTarget->InitAutoFormat(RenderTargetSize.X, RenderTargetSize.Y);
 	RenderTarget->UpdateResourceImmediate(true);
 
 	// Assign the RT to the compositing output
@@ -599,6 +620,127 @@ void FCameraCalibrationStepsController::CreateComp()
 	}
 }
 
+bool FCameraCalibrationStepsController::IsOverlayEnabled() const
+{
+	if (UCompositingElementMaterialPass* OverlayPassPtr = OverlayPass.Get())
+	{
+		return OverlayPassPtr->IsPassEnabled();
+	}
+
+	return false;
+}
+
+void FCameraCalibrationStepsController::SetOverlayEnabled(const bool bEnabled)
+{
+	if (UCompositingElementMaterialPass* OverlayPassPtr = OverlayPass.Get())
+	{
+		OverlayPassPtr->SetPassEnabled(bEnabled);
+	}
+}
+
+void FCameraCalibrationStepsController::SetOverlayMaterial(UMaterialInterface* OverlayMaterial, bool bShowOverlay)
+{
+	if (UCompositingElementMaterialPass* OverlayPassPtr = OverlayPass.Get())
+	{
+		OverlayPassPtr->SetMaterialInterface(OverlayMaterial);
+		OverlayPassPtr->SetPassEnabled(bShowOverlay);
+	}
+}
+
+const UMaterialInterface* const FCameraCalibrationStepsController::GetOverlayMaterial() const
+{
+	if (UCompositingElementMaterialPass* OverlayPassPtr = OverlayPass.Get())
+	{
+		if (OverlayPassPtr->IsPassEnabled())
+		{
+			return OverlayPassPtr->Material.Material;
+		}
+	}
+
+	return nullptr;
+}
+
+bool FCameraCalibrationStepsController::GetOverlayScalarParameterValue(const FName& ParameterName, float& OutValue) const
+{
+	bool Result = false;
+
+	if (UCompositingElementMaterialPass* OverlayPassPtr = OverlayPass.Get())
+	{
+		const UMaterialInstanceDynamic* const OverlayMID = OverlayPassPtr->Material.GetMID();
+		Result = OverlayMID->GetScalarParameterValue(ParameterName, OutValue);
+	}
+
+	return Result;
+}
+
+bool FCameraCalibrationStepsController::GetOverlayScalarParameterMinMax(const FName& ParameterName, float& OutMinValue, float& OutMaxValue) const
+{
+	bool Result = false;
+
+	if (UCompositingElementMaterialPass* OverlayPassPtr = OverlayPass.Get())
+	{
+		const UMaterialInstanceDynamic* const OverlayMID = OverlayPassPtr->Material.GetMID();
+		Result = OverlayMID->GetScalarParameterSliderMinMax(ParameterName, OutMinValue, OutMaxValue);
+	}
+
+	return Result;
+}
+
+bool FCameraCalibrationStepsController::GetOverlayVectorParameterValue(const FName& ParameterName, FLinearColor& OutValue) const
+{
+	bool Result = false;
+
+	if (UCompositingElementMaterialPass* OverlayPassPtr = OverlayPass.Get())
+	{
+		const UMaterialInstanceDynamic* const OverlayMID = OverlayPassPtr->Material.GetMID();
+		Result = OverlayMID->GetVectorParameterValue(ParameterName, OutValue);
+	}
+
+	return Result;
+}
+
+bool FCameraCalibrationStepsController::GetOverlayTextureParameterValue(const FName& ParameterName, UTexture*& OutValue) const
+{
+	bool Result = false;
+
+	if (UCompositingElementMaterialPass* OverlayPassPtr = OverlayPass.Get())
+	{
+		Result = OverlayPassPtr->Material.GetTextureOverride(ParameterName, OutValue);
+
+		if (!Result)
+		{
+			const UMaterialInstanceDynamic* const OverlayMID = OverlayPassPtr->Material.GetMID();
+			Result = OverlayMID->GetTextureParameterValue(ParameterName, OutValue);
+		}
+	}
+
+	return Result;
+}
+
+void FCameraCalibrationStepsController::SetOverlayScalarParameter(const FName& ParameterName, const float NewValue)
+{
+	if (UCompositingElementMaterialPass* OverlayPassPtr = OverlayPass.Get())
+	{
+		OverlayPassPtr->Material.SetScalarOverride(ParameterName, NewValue);
+	}
+}
+
+void FCameraCalibrationStepsController::SetOverlayVectorParameter(const FName& ParameterName, const FLinearColor& NewValue)
+{
+	if (UCompositingElementMaterialPass* OverlayPassPtr = OverlayPass.Get())
+	{
+		OverlayPassPtr->Material.SetVectorOverride(ParameterName, NewValue);
+	}
+}
+
+void FCameraCalibrationStepsController::SetOverlayTextureParameter(const FName& ParameterName, UTexture* NewValue)
+{
+	if (UCompositingElementMaterialPass* OverlayPassPtr = OverlayPass.Get())
+	{
+		OverlayPassPtr->Material.SetTextureOverride(ParameterName, NewValue);
+	}
+}
+
 void FCameraCalibrationStepsController::CreateMediaPlateOutput()
 {
 	if (!MediaPlate.IsValid())
@@ -634,7 +776,7 @@ void FCameraCalibrationStepsController::CreateMediaPlateOutput()
 	MediaPlateRenderTarget->RenderTargetFormat = ETextureRenderTargetFormat::RTF_RGBA8;
 	MediaPlateRenderTarget->ClearColor = FLinearColor::Black;
 	MediaPlateRenderTarget->bAutoGenerateMips = true;
-	MediaPlateRenderTarget->InitAutoFormat(1920, 1080);
+	MediaPlateRenderTarget->InitAutoFormat(RenderTargetSize.X, RenderTargetSize.Y);
 	MediaPlateRenderTarget->UpdateResourceImmediate(true);
 
 	// Assign the RT to the compositing output
diff --git a/Engine/Plugins/VirtualProduction/CameraCalibration/Source/CameraCalibrationEditor/Private/AssetEditor/CameraCalibrationStepsController.h b/Engine/Plugins/VirtualProduction/CameraCalibration/Source/CameraCalibrationEditor/Private/AssetEditor/CameraCalibrationStepsController.h
--- a/Engine/Plugins/VirtualProduction/CameraCalibration/Source/CameraCalibrationEditor/Private/AssetEditor/CameraCalibrationStepsController.h
+++ b/Engine/Plugins/VirtualProduction/CameraCalibration/Source/CameraCalibrationEditor/Private/AssetEditor/CameraCalibrationStepsController.h
@@ -52,6 +52,9 @@ public:
 	/** Returns the render target of the Media Plate */
 	UTextureRenderTarget2D* GetMediaPlateRenderTarget() const;
 
+	/** Returns the size of the render target used by the Comp */
+	FIntPoint GetCompRenderTargetSize() const;
+
 	/** Creates a way to read the media plate pixels for processing by any calibration step */
 	void CreateMediaPlateOutput();
 
@@ -118,6 +121,39 @@ public:
 	/** Reads the pixels in the media plate */
 	bool ReadMediaPixels(TArray<FColor>& Pixels, FIntPoint& Size, ETextureRenderTargetFormat& PixelFormat, FText& OutErrorMessage) const;
 
+	/** Returns true if the overlay transform pass is currently enabled */
+	bool IsOverlayEnabled() const;
+
+	/** Sets the enabled state of the overlay transform pass */
+	void SetOverlayEnabled(const bool bEnabled = true);
+
+	/** Sets the overlay material to be used by the overlay transform pass */
+	void SetOverlayMaterial(UMaterialInterface* OverlayMaterial, bool bShowOverlay = true);
+
+	/** Gets the overlay material currently used by the overlay transform pass */
+	const UMaterialInterface* const GetOverlayMaterial() const;
+
+	/** Get the current value of the input parameter from the overlay material */
+	bool GetOverlayScalarParameterValue(const FName& ParameterName, float& OutValue) const;
+
+	/** Get the min value of the input parameter from the overlay material */
+	bool GetOverlayScalarParameterMinMax(const FName& ParameterName, float& OutMinValue, float& OutMaxValue) const;
+
+	/** Get the value of the input parameter from the overlay material */
+	bool GetOverlayVectorParameterValue(const FName& ParameterName, FLinearColor& OutValue) const;
+
+	/** Get the value of the input parameter from the overlay material */
+	bool GetOverlayTextureParameterValue(const FName& ParameterName, UTexture*& OutValue) const;
+
+	/** Set the value of the input scalar parameter on the overlay material */
+	void SetOverlayScalarParameter(const FName& ParameterName, const float NewValue);
+
+	/** Set the value of the input vector parameter on the overlay material */
+	void SetOverlayVectorParameter(const FName& ParameterName, const FLinearColor& NewValue);
+
+	/** Set the value of the input texture parameter on the overlay material */
+	void SetOverlayTextureParameter(const FName& ParameterName, UTexture* NewValue);
+
 public:
 
 	/** Called by the UI when the Simulcam Viewport is clicked */
@@ -160,6 +196,9 @@ private:
 	/** Pointer to the camera calibration toolkit */
 	TWeakPtr<FCameraCalibrationToolkit> CameraCalibrationToolkit;
 
+	/** Size to use when creating the render targets for the comp and media output */
+	FIntPoint RenderTargetSize;
+
 	/** Array of the calibration steps that this controller is managing */
 	TArray<TStrongObjectPtr<UCameraCalibrationStep>> CalibrationSteps;
 
@@ -190,6 +229,9 @@ private:
 	/** The material pass the does the CG + MediaPlate composite with a wiper weight */
 	TWeakObjectPtr<UCompositingElementMaterialPass> MaterialPass;
 
+	/** The material pass that renders an overlay on top of the composite */
+	TWeakObjectPtr<UCompositingElementMaterialPass> OverlayPass;
+
 	/** The currently selected camera */
 	TWeakObjectPtr<ACameraActor> Camera;
 
diff --git a/Engine/Plugins/VirtualProduction/CameraCalibration/Source/CameraCalibrationEditor/Private/AssetEditor/SCameraCalibrationSteps.cpp b/Engine/Plugins/VirtualProduction/CameraCalibration/Source/CameraCalibrationEditor/Private/AssetEditor/SCameraCalibrationSteps.cpp
--- a/Engine/Plugins/VirtualProduction/CameraCalibration/Source/CameraCalibrationEditor/Private/AssetEditor/SCameraCalibrationSteps.cpp
+++ b/Engine/Plugins/VirtualProduction/CameraCalibration/Source/CameraCalibrationEditor/Private/AssetEditor/SCameraCalibrationSteps.cpp
@@ -18,9 +18,11 @@
 #include "Styling/CoreStyle.h"
 #include "Styling/SlateStyle.h"
 #include "UI/CameraCalibrationWidgetHelpers.h"
+#include "Widgets/Colors/SColorPicker.h"
 #include "Widgets/Images/SImage.h"
 #include "Widgets/Input/SCheckBox.h"
 #include "Widgets/Input/SComboBox.h"
+#include "Widgets/Input/SNumericEntryBox.h"
 #include "Widgets/Input/SSpinBox.h"
 #include "Widgets/Layout/SSpacer.h"
 #include "Widgets/Layout/SWidgetSwitcher.h"
@@ -117,6 +119,16 @@ void SCameraCalibrationSteps::Construct(const FArguments& InArgs, TWeakPtr<FCame
 			.MaxHeight(FCameraCalibrationWidgetHelpers::DefaultRowHeight)
 			[FCameraCalibrationWidgetHelpers::BuildLabelWidgetPair(LOCTEXT("MediaSource", "Media Source"), BuildMediaSourceWidget())]
 
+			+ SVerticalBox::Slot() // Overlay picker
+			.MaxHeight(FCameraCalibrationWidgetHelpers::DefaultRowHeight)
+			[FCameraCalibrationWidgetHelpers::BuildLabelWidgetPair(LOCTEXT("Overlay", "Overlay"), BuildOverlayWidget())]
+
+			+ SVerticalBox::Slot() // Overlay parameters
+			.AutoHeight()
+			[
+				OverlayParameterWidget.ToSharedRef()
+			]
+
 			+ SVerticalBox::Slot() // Step Title
 			.Padding(0, 5)
 			.AutoHeight()
@@ -290,6 +302,257 @@ TSharedRef<SWidget> SCameraCalibrationSteps::BuildMediaSourceWidget()
 	return MediaSourcesComboBox.ToSharedRef();
 }
 
+TSharedRef<SWidget> SCameraCalibrationSteps::BuildOverlayWidget()
+{
+	const UCameraCalibrationSubsystem* const SubSystem = GEngine->GetEngineSubsystem<UCameraCalibrationSubsystem>();
+
+	for (const FName& Name : SubSystem->GetOverlayMaterialNames())
+	{
+		SharedOverlayNames.Add(MakeShared<FName>(Name));
+	}
+
+	SharedOverlayNames.Sort([](const TSharedPtr<FName>& LHS, const TSharedPtr<FName>& RHS) { return LHS->Compare(*RHS) <= 0; });
+
+	SharedOverlayNames.Insert(MakeShared<FName>(FName(TEXT("None"))), 0);
+
+	OverlayComboBox = SNew(SComboBox<TSharedPtr<FName>>)
+		.OptionsSource(&SharedOverlayNames)
+		.OnSelectionChanged_Lambda([&](TSharedPtr<FName> NewValue, ESelectInfo::Type Type) -> void
+		{
+			if (!CalibrationStepsController.IsValid())
+			{
+				return;
+			}
+				
+			if (NewValue->Compare(FName(TEXT("None"))) == 0.0f)
+			{
+				CalibrationStepsController.Pin()->SetOverlayEnabled(false);
+			}
+			else
+			{
+				UMaterialInterface* OverlayMaterial = SubSystem->GetOverlayMaterial(*NewValue);
+				CalibrationStepsController.Pin()->SetOverlayMaterial(OverlayMaterial);
+			}
+
+			UpdateOverlayMaterialParameterWidget();
+		})
+		.OnGenerateWidget_Lambda([&](TSharedPtr<FName> InOption) -> TSharedRef<SWidget>
+		{
+			return SNew(STextBlock).Text(FText::FromName(*InOption));
+		})
+		.InitiallySelectedItem(nullptr)
+		[
+			SNew(STextBlock)
+			.Text_Lambda([&]() -> FText
+			{
+				if (OverlayComboBox.IsValid() && OverlayComboBox->GetSelectedItem().IsValid())
+				{
+					return FText::FromName(*OverlayComboBox->GetSelectedItem());
+				}
+
+				return LOCTEXT("InvalidComboOption", "None");
+			})
+		];
+
+	OverlayParameterWidget = SNew(SHorizontalBox);
+	OverlayParameterListWidget = SNew(SVerticalBox);
+
+	return OverlayComboBox.ToSharedRef();
+}
+
+void SCameraCalibrationSteps::UpdateOverlayMaterialParameterWidget()
+{
+	OverlayParameterWidget->ClearChildren();
+	OverlayParameterListWidget->ClearChildren();
+
+	if (!CalibrationStepsController.IsValid())
+	{
+		return;
+	}
+
+	if (CalibrationStepsController.Pin()->IsOverlayEnabled())
+	{
+		if (const UMaterialInterface* const OverlayMaterial = CalibrationStepsController.Pin()->GetOverlayMaterial())
+		{
+			TMap<FMaterialParameterInfo, FMaterialParameterMetadata> ScalarParams;
+			OverlayMaterial->GetAllParametersOfType(EMaterialParameterType::Scalar, ScalarParams);
+
+			TMap<FMaterialParameterInfo, FMaterialParameterMetadata> VectorParams;
+			OverlayMaterial->GetAllParametersOfType(EMaterialParameterType::Vector, VectorParams);
+
+			// Early-exit if there are no material parameters to display
+			if (ScalarParams.Num() + VectorParams.Num() == 0)
+			{
+				return;
+			}
+
+			for (const TPair<FMaterialParameterInfo, FMaterialParameterMetadata>& Param : ScalarParams)
+			{
+				const FMaterialParameterInfo ParameterInfo = Param.Key;
+				const FMaterialParameterMetadata ParameterData = Param.Value;
+
+				OverlayParameterListWidget->AddSlot()
+				.Padding(5, 5)
+				[
+					SNew(SHorizontalBox)
+
+					+ SHorizontalBox::Slot()
+					[
+						SNew(STextBlock)
+						.Text(FText::FromString(ParameterInfo.Name.ToString()))
+					]
+
+					+ SHorizontalBox::Slot()
+					[
+						SNew(SNumericEntryBox<float>)
+						.AllowSpin(true)
+						.MinValue(TOptional<float>())
+						.MaxValue(TOptional<float>())
+						.Delta(0.0f)
+						.Value_Lambda([=]() 
+						{ 
+							float ScalarValue = 0.0f;
+							
+							if (CalibrationStepsController.IsValid())
+							{
+								CalibrationStepsController.Pin()->GetOverlayScalarParameterValue(ParameterInfo.Name, ScalarValue);
+							}
+
+							return ScalarValue;
+						})
+						.MinSliderValue_Lambda([=]() 
+						{ 
+							float MinValue = 0.0f;
+							float MaxValue = 0.0f;
+
+							if (CalibrationStepsController.IsValid())
+							{
+								CalibrationStepsController.Pin()->GetOverlayScalarParameterMinMax(ParameterInfo.Name, MinValue, MaxValue);
+							}
+
+							return MinValue;
+						})
+						.MaxSliderValue_Lambda([=]() 
+						{ 
+							float MinValue = 0.0f;
+							float MaxValue = 0.0f;
+								
+							if (CalibrationStepsController.IsValid())
+							{
+								CalibrationStepsController.Pin()->GetOverlayScalarParameterMinMax(ParameterInfo.Name, MinValue, MaxValue);
+							}
+
+							return MaxValue;
+						})
+						.OnValueChanged_Lambda([=](float NewValue) 
+						{ 
+							if (CalibrationStepsController.IsValid())
+							{
+								CalibrationStepsController.Pin()->SetOverlayScalarParameter(ParameterInfo.Name, NewValue);
+							}
+						})
+						.OnValueCommitted_Lambda([=](float NewValue, ETextCommit::Type)	
+						{
+							if (CalibrationStepsController.IsValid())
+							{
+								CalibrationStepsController.Pin()->SetOverlayScalarParameter(ParameterInfo.Name, NewValue);
+							}
+						})
+					]
+				]
+				;
+			}
+
+			for (const TPair<FMaterialParameterInfo, FMaterialParameterMetadata>& Param : VectorParams)
+			{
+				const FMaterialParameterInfo ParameterInfo = Param.Key;
+				const FMaterialParameterMetadata ParameterData = Param.Value;
+
+				OverlayParameterListWidget->AddSlot()
+				.Padding(5, 5)
+				.AutoHeight()
+				[
+					SNew(SHorizontalBox)
+
+					+ SHorizontalBox::Slot()
+					[
+						SNew(STextBlock)
+						.Text(FText::FromString(ParameterInfo.Name.ToString()))
+					]
+
+					+ SHorizontalBox::Slot()
+					[
+						SNew(SComboButton)
+						.ContentPadding(0)
+						.HasDownArrow(false)
+						.CollapseMenuOnParentFocus(true)
+						.ButtonStyle(FEditorStyle::Get(), "Sequencer.AnimationOutliner.ColorStrip") // Style matches the button used in cinematic film overlays
+						.OnGetMenuContent_Lambda([=]() -> TSharedRef<SWidget>
+						{
+							return SNew(SColorPicker)
+								.UseAlpha(true)
+								.TargetColorAttribute_Lambda([=]() 
+								{
+									FLinearColor ColorValue = FLinearColor(1.0f, 1.0f, 1.0f, 1.0f);
+										
+									if (CalibrationStepsController.IsValid())
+									{
+										CalibrationStepsController.Pin()->GetOverlayVectorParameterValue(ParameterInfo.Name, ColorValue);
+									}
+
+									return ColorValue;
+								})
+								.OnColorCommitted_Lambda([=](FLinearColor NewColor) 
+								{
+									if (CalibrationStepsController.IsValid())
+									{
+										CalibrationStepsController.Pin()->SetOverlayVectorParameter(ParameterInfo.Name, NewColor);
+									}
+								});
+						})
+						.ButtonContent()
+						[
+							SNew(SColorBlock)
+							.ShowBackgroundForAlpha(true)
+							.Size(FVector2D(10.0f, 20.0f))
+							.Color_Lambda([=]()
+							{
+								FLinearColor ColorValue = FLinearColor(1.0f, 1.0f, 1.0f, 1.0f);
+
+								if (CalibrationStepsController.IsValid())
+								{
+									CalibrationStepsController.Pin()->GetOverlayVectorParameterValue(ParameterInfo.Name, ColorValue);
+								}
+
+								return ColorValue;
+							})
+						]
+					]
+				]
+				;
+			}
+
+			OverlayParameterWidget->AddSlot()
+				.VAlign(VAlign_Top)
+				.Padding(5, 10)
+				.FillWidth(0.35f)
+				[
+					SNew(STextBlock).Text(LOCTEXT("OverlayParams", "Overlay Parameters"))
+				]
+			;
+
+			OverlayParameterWidget->AddSlot()
+				.VAlign(VAlign_Center)
+				.Padding(10, 5)
+				.FillWidth(0.65f)
+				[
+					OverlayParameterListWidget.ToSharedRef()
+				]
+			;
+		}
+	}
+}
+
 TSharedRef<SWidget> SCameraCalibrationSteps::BuildStepSelectionWidget()
 {
 	if (!CalibrationStepsController.IsValid())
diff --git a/Engine/Plugins/VirtualProduction/CameraCalibration/Source/CameraCalibrationEditor/Private/AssetEditor/SCameraCalibrationSteps.h b/Engine/Plugins/VirtualProduction/CameraCalibration/Source/CameraCalibrationEditor/Private/AssetEditor/SCameraCalibrationSteps.h
--- a/Engine/Plugins/VirtualProduction/CameraCalibration/Source/CameraCalibrationEditor/Private/AssetEditor/SCameraCalibrationSteps.h
+++ b/Engine/Plugins/VirtualProduction/CameraCalibration/Source/CameraCalibrationEditor/Private/AssetEditor/SCameraCalibrationSteps.h
@@ -41,9 +41,15 @@ private:
 	/** Builds the UI for the media source picker */
 	TSharedRef<SWidget> BuildMediaSourceWidget();
 
+	/** Builds the UI for the overlay picker */
+	TSharedRef<SWidget> BuildOverlayWidget();
+
 	/** Builds the UI for the calibration step selection */
 	TSharedRef<SWidget> BuildStepSelectionWidget();
 
+	/** Updates the material parameter widget to display the parameters for the currently selected overlay */
+	void UpdateOverlayMaterialParameterWidget();
+
 	/** Refreshes the list of available media sources shown in the MediaSourcesComboBox */
 	void UpdateMediaSourcesOptions();
 
@@ -61,6 +67,18 @@ private:
 	/** The combobox that presents the available media sources */
 	TSharedPtr<SComboBox<TSharedPtr<FString>>> MediaSourcesComboBox;
 
+	/** The combox that presents the available overlays */
+	TSharedPtr<SComboBox<TSharedPtr<FName>>> OverlayComboBox;
+
+	/** The overlay parameter widget that displays a category name and the parameter list */
+	TSharedPtr<SHorizontalBox> OverlayParameterWidget;
+
+	/** The overlay parameter list widget that display the scalar and vector parameters for the currently select overlay */
+	TSharedPtr<SVerticalBox> OverlayParameterListWidget;
+
+	/** List of overlay names for the combox box options */
+	TArray<TSharedPtr<FName>> SharedOverlayNames;
+
 	/** Widget switcher to only display the UI of the selected step */
 	TSharedPtr<SWidgetSwitcher> StepWidgetSwitcher;
 
diff --git a/Engine/Plugins/VirtualProduction/CameraCalibration/Source/CameraCalibrationEditor/Private/Calibrators/CameraLensDistortionAlgoCheckerboard.cpp b/Engine/Plugins/VirtualProduction/CameraCalibration/Source/CameraCalibrationEditor/Private/Calibrators/CameraLensDistortionAlgoCheckerboard.cpp
--- a/Engine/Plugins/VirtualProduction/CameraCalibration/Source/CameraCalibrationEditor/Private/Calibrators/CameraLensDistortionAlgoCheckerboard.cpp
+++ b/Engine/Plugins/VirtualProduction/CameraCalibration/Source/CameraCalibrationEditor/Private/Calibrators/CameraLensDistortionAlgoCheckerboard.cpp
@@ -16,19 +16,6 @@
 #include "Widgets/Input/SCheckBox.h"
 
 
-#if WITH_OPENCV
-
-#include <vector>
-
-#include "OpenCVHelper.h"
-OPENCV_INCLUDES_START
-#undef check 
-#include "opencv2/opencv.hpp"
-#include "opencv2/calib3d.hpp"
-OPENCV_INCLUDES_END
-
-#endif //WITH_OPENCV
-
 #define LOCTEXT_NAMESPACE "CameraLensDistortionAlgoCheckerboard"
 
 
@@ -109,6 +96,19 @@ void UCameraLensDistortionAlgoCheckerboard::Initialize(ULensDistortionTool* InTo
 
 	// Guess which calibrator to use.
 	SetCalibrator(FindFirstCalibrator());
+
+#if WITH_OPENCV
+	// Initialize coverage matrix
+	FCameraCalibrationStepsController* StepsController = Tool->GetCameraCalibrationStepsController();
+
+	if (!ensure(StepsController))
+	{
+		return;
+	}
+
+	FIntPoint Size = StepsController->GetCompRenderTargetSize();
+	CvCoverage = cv::Mat(cv::Size(Size.X, Size.Y), CV_8UC4);
+#endif
 }
 
 void UCameraLensDistortionAlgoCheckerboard::Shutdown()
@@ -301,16 +301,26 @@ bool UCameraLensDistortionAlgoCheckerboard::AddCalibrationRow(FText& OutErrorMes
 			Row->Points2d.Add(FVector2D(Corner.x, Corner.y));
 		}
 
-		cv::drawChessboardCorners(CvFrame, CheckerboardSize, Corners, bCornersFound);
-	}
+		// Update the coverage overlay image with information from the newly added row
+		cv::drawChessboardCorners(CvCoverage, CheckerboardSize, Corners, true);
 
-	// Show the detection to the user
-	if (bShouldShowDetectionWindow)
-	{
-		FCameraCalibrationWidgetHelpers::DisplayTextureInWindowAlmostFullScreen(
-			FOpenCVHelper::TextureFromCvMat(CvFrame),
-			LOCTEXT("CheckerboardDetection", "Checkerboard Detection")
-		);
+		CoverageTexture = FOpenCVHelper::TextureFromCvMat(CvCoverage, CoverageTexture);
+
+		UMaterialInterface* CoverageOverlay = TSoftObjectPtr<UMaterialInterface>(FSoftObjectPath(TEXT("/CameraCalibration/Materials/M_Coverage.M_Coverage"))).LoadSynchronous();
+		StepsController->SetOverlayMaterial(CoverageOverlay, bShouldShowOverlay);
+		StepsController->SetOverlayTextureParameter(FName(TEXT("CoverageTexture")), CoverageTexture);
+
+		// Show the detection to the user
+		if (bShouldShowDetectionWindow)
+		{
+			// Update the coverage overlay image with information from the newly added row
+			cv::drawChessboardCorners(CvFrame, CheckerboardSize, Corners, true);
+
+			FCameraCalibrationWidgetHelpers::DisplayTextureInWindowAlmostFullScreen(
+				FOpenCVHelper::TextureFromCvMat(CvFrame),
+				LOCTEXT("CheckerboardDetection", "Checkerboard Detection")
+			);
+		}
 	}
 
 	// Create thumbnail and add it to the row
@@ -362,6 +372,12 @@ TSharedRef<SWidget> UCameraLensDistortionAlgoCheckerboard::BuildUI()
 		.MaxHeight(FCameraCalibrationWidgetHelpers::DefaultRowHeight)
 		[ FCameraCalibrationWidgetHelpers::BuildLabelWidgetPair(LOCTEXT("Checkerboard", "Checkerboard"), BuildCalibrationDevicePickerWidget()) ]
 
+		+ SVerticalBox::Slot() // Show Overlay
+		.VAlign(EVerticalAlignment::VAlign_Top)
+		.AutoHeight()
+		.MaxHeight(FCameraCalibrationWidgetHelpers::DefaultRowHeight)
+		[FCameraCalibrationWidgetHelpers::BuildLabelWidgetPair(LOCTEXT("ShowOverlay", "Show Coverage Overlay"), BuildShowOverlayWidget())]
+
 		+ SVerticalBox::Slot() // Show Detection
 		.VAlign(EVerticalAlignment::VAlign_Top)
 		.AutoHeight()
@@ -708,6 +724,28 @@ TSharedRef<SWidget> UCameraLensDistortionAlgoCheckerboard::BuildCalibrationDevic
 		;
 }
 
+TSharedRef<SWidget> UCameraLensDistortionAlgoCheckerboard::BuildShowOverlayWidget()
+{
+	return SNew(SCheckBox)
+		.IsChecked_Lambda([&]() -> ECheckBoxState
+		{
+			return bShouldShowOverlay ? ECheckBoxState::Checked : ECheckBoxState::Unchecked;
+		})
+		.OnCheckStateChanged_Lambda([&](ECheckBoxState NewState) -> void
+		{
+			bShouldShowOverlay = (NewState == ECheckBoxState::Checked);
+
+			FCameraCalibrationStepsController* StepsController = Tool->GetCameraCalibrationStepsController();
+
+			if (!ensure(StepsController))
+			{
+				return;
+			}
+
+			StepsController->SetOverlayEnabled(bShouldShowOverlay);
+		});
+}
+
 TSharedRef<SWidget> UCameraLensDistortionAlgoCheckerboard::BuildShowDetectionWidget()
 {
 	return SNew(SCheckBox)
@@ -771,6 +809,7 @@ TSharedRef<SWidget> UCameraLensDistortionAlgoCheckerboard::BuildCalibrationPoint
 				}
 
 				CalibrationListView->RequestListRefresh();
+				RefreshCoverage();
 				return FReply::Handled();
 			}
 			else if (KeyEvent.GetModifierKeys().IsControlDown() && (KeyEvent.GetKey() == EKeys::A))
@@ -859,6 +898,45 @@ void UCameraLensDistortionAlgoCheckerboard::ClearCalibrationRows()
 	{
 		CalibrationListView->RequestListRefresh();
 	}
+
+	RefreshCoverage();
+}
+
+void UCameraLensDistortionAlgoCheckerboard::RefreshCoverage()
+{
+#if WITH_OPENCV
+	FCameraCalibrationStepsController* StepsController = Tool->GetCameraCalibrationStepsController();
+
+	if (!ensure(StepsController))
+	{
+		return;
+	}
+
+	FIntPoint Size = StepsController->GetCompRenderTargetSize();
+
+	CvCoverage.release();
+	CvCoverage = cv::Mat(cv::Size(Size.X, Size.Y), CV_8UC4);
+
+	for (const TSharedPtr<FCalibrationRowData>& Row : CalibrationRows)
+	{
+		cv::Size CheckerboardSize(Row->NumCornerCols, Row->NumCornerRows);
+
+		std::vector<cv::Point2f> Corners;
+
+		for (const FVector2D& Corner : Row->Points2d)
+		{
+			Corners.push_back(cv::Point2f(Corner.X, Corner.Y));
+		}
+
+		cv::drawChessboardCorners(CvCoverage, CheckerboardSize, Corners, true);
+	}
+
+	CoverageTexture = FOpenCVHelper::TextureFromCvMat(CvCoverage, CoverageTexture);
+
+	UMaterialInterface* CoverageOverlay = TSoftObjectPtr<UMaterialInterface>(FSoftObjectPath(TEXT("/CameraCalibration/Materials/M_Coverage.M_Coverage"))).LoadSynchronous();
+	StepsController->SetOverlayMaterial(CoverageOverlay, bShouldShowOverlay);
+	StepsController->SetOverlayTextureParameter(FName(TEXT("CoverageTexture")), CoverageTexture);
+#endif
 }
 
 TSharedRef<SWidget> UCameraLensDistortionAlgoCheckerboard::BuildHelpWidget()
diff --git a/Engine/Plugins/VirtualProduction/CameraCalibration/Source/CameraCalibrationEditor/Private/Calibrators/CameraLensDistortionAlgoCheckerboard.h b/Engine/Plugins/VirtualProduction/CameraCalibration/Source/CameraCalibrationEditor/Private/Calibrators/CameraLensDistortionAlgoCheckerboard.h
--- a/Engine/Plugins/VirtualProduction/CameraCalibration/Source/CameraCalibrationEditor/Private/Calibrators/CameraLensDistortionAlgoCheckerboard.h
+++ b/Engine/Plugins/VirtualProduction/CameraCalibration/Source/CameraCalibrationEditor/Private/Calibrators/CameraLensDistortionAlgoCheckerboard.h
@@ -6,6 +6,20 @@
 
 #include "LensFile.h"
 
+#if WITH_OPENCV
+
+#include <vector>
+
+#include "OpenCVHelper.h"
+OPENCV_INCLUDES_START
+#undef check 
+#include "opencv2/opencv.hpp"
+#include "opencv2/calib3d.hpp"
+OPENCV_INCLUDES_END
+
+#endif //WITH_OPENCV
+
+
 #include "CameraLensDistortionAlgoCheckerboard.generated.h"
 
 struct FGeometry;
@@ -132,14 +146,28 @@ private:
 	/** Caches the last camera data.  Will hold last value before the media is paused */
 	FCameraDataCache LastCameraData;
 
+	/** True if the coverage overlay should be shown */
+	bool bShouldShowOverlay = false;
+
 	/** True if a detection window should be shown after every capture */
 	bool bShouldShowDetectionWindow = false;
 
+	/** Texture into which detected chessboard corners from each calibration row are drawn */
+	TObjectPtr<UTexture2D> CoverageTexture;
+
+#if WITH_OPENCV
+	/** OpenCV matrix used to draw chessboard corners */
+	cv::Mat CvCoverage;
+#endif
+
 private:
 
 	/** Builds the UI of the calibration device picker */
 	TSharedRef<SWidget> BuildCalibrationDevicePickerWidget();
 
+	/** Builds the UI for the user to select if they want to display the coverage overlay */
+	TSharedRef<SWidget> BuildShowOverlayWidget();
+
 	/** Builds the UI for the user to select if they want a corner detection window to be shown after every capture */
 	TSharedRef<SWidget> BuildShowDetectionWidget();
 
@@ -171,4 +199,7 @@ private:
 
 	/** Returns the steps controller */
 	FCameraCalibrationStepsController* GetStepsController() const;
+
+	/** Update the coverage texture to reflect the current set of calibration rows */
+	void RefreshCoverage();
 };
diff --git a/Engine/Plugins/VirtualProduction/CameraCalibrationCore/Source/CameraCalibrationCore/Private/CameraCalibrationSettings.cpp b/Engine/Plugins/VirtualProduction/CameraCalibrationCore/Source/CameraCalibrationCore/Private/CameraCalibrationSettings.cpp
--- a/Engine/Plugins/VirtualProduction/CameraCalibrationCore/Source/CameraCalibrationCore/Private/CameraCalibrationSettings.cpp
+++ b/Engine/Plugins/VirtualProduction/CameraCalibrationCore/Source/CameraCalibrationCore/Private/CameraCalibrationSettings.cpp
@@ -78,6 +78,29 @@ void UCameraCalibrationSettings::PostEditChangeChainProperty(struct FPropertyCha
 	}
 }
 
+TArray<FName> UCameraCalibrationSettings::GetDefaultCalibrationOverlayNames() const
+{
+	TArray<FName> OutKeys;
+#if WITH_EDITORONLY_DATA
+	DefaultCalibrationOverlayMaterials.GetKeys(OutKeys);
+#endif
+	return OutKeys;
+}
+
+UMaterialInterface* UCameraCalibrationSettings::GetDefaultCalibrationOverlayMaterial(const FName OverlayName) const
+{
+#if WITH_EDITORONLY_DATA
+	const TSoftObjectPtr<UMaterialInterface>* OverlayMaterial = DefaultCalibrationOverlayMaterials.Find(OverlayName);
+	if (!OverlayMaterial)
+	{
+		return nullptr;
+	}
+	return OverlayMaterial->LoadSynchronous();
+#else 
+	return nullptr;
+#endif
+}
+
 FOnDisplacementMapResolutionChanged& UCameraCalibrationSettings::OnDisplacementMapResolutionChanged()
 {
 	return DisplacementMapResolutionChangedDelegate;
@@ -89,6 +112,10 @@ FOnCalibrationInputToleranceChanged& UCameraCalibrationSettings::OnCalibrationIn
 }
 #endif
 
+UCameraCalibrationEditorSettings::UCameraCalibrationEditorSettings()
+{
+}
+
 FName UCameraCalibrationEditorSettings::GetCategoryName() const
 {
 	return TEXT("Plugins");
diff --git a/Engine/Plugins/VirtualProduction/CameraCalibrationCore/Source/CameraCalibrationCore/Private/CameraCalibrationSubsystem.cpp b/Engine/Plugins/VirtualProduction/CameraCalibrationCore/Source/CameraCalibrationCore/Private/CameraCalibrationSubsystem.cpp
--- a/Engine/Plugins/VirtualProduction/CameraCalibrationCore/Source/CameraCalibrationCore/Private/CameraCalibrationSubsystem.cpp
+++ b/Engine/Plugins/VirtualProduction/CameraCalibrationCore/Source/CameraCalibrationCore/Private/CameraCalibrationSubsystem.cpp
@@ -3,6 +3,7 @@
 #include "CameraCalibrationSubsystem.h"
 
 #include "CameraCalibrationCoreLog.h"
+#include "CameraCalibrationSettings.h"
 #include "CineCameraComponent.h"
 #include "Engine/TimecodeProvider.h"
 #include "Misc/CoreDelegates.h"
@@ -160,6 +161,16 @@ void UCameraCalibrationSubsystem::UnregisterDistortionModel(TSubclassOf<ULensMod
 	LensModelMap.Remove(LensModel->GetDefaultObject<ULensModel>()->GetModelName());
 }
 
+void UCameraCalibrationSubsystem::RegisterOverlayMaterial(const FName& MaterialName, const FName& MaterialPath)
+{
+	RegisteredOverlayMaterials.Add(MaterialName, TSoftObjectPtr<UMaterialInterface>(FSoftObjectPath(MaterialPath)));
+}
+
+void UCameraCalibrationSubsystem::UnregisterOverlayMaterial(const FName& MaterialName)
+{
+	RegisteredOverlayMaterials.Remove(MaterialName);
+}
+
 void UCameraCalibrationSubsystem::UpdateOriginalFocalLength(UCineCameraComponent* Component, float InFocalLength)
 {
 	FCachedFocalLength* CachedFocalLength = CachedFocalLengthMap.Find(Component);
@@ -234,6 +245,31 @@ TArray<FName> UCameraCalibrationSubsystem::GetCameraImageCenterAlgos() const
 	return OutKeys;
 }
 
+UMaterialInterface* UCameraCalibrationSubsystem::GetOverlayMaterial(const FName& OverlayName) const
+{
+#if WITH_EDITOR
+	if (UMaterialInterface* DefaultMaterial = GetDefault<UCameraCalibrationSettings>()->GetDefaultCalibrationOverlayMaterial(OverlayName))
+	{
+		return DefaultMaterial;
+	}
+#endif
+
+	const TSoftObjectPtr<UMaterialInterface> OverlayMaterial = RegisteredOverlayMaterials.FindRef(OverlayName);
+
+	return OverlayMaterial.LoadSynchronous();
+}
+
+TArray<FName> UCameraCalibrationSubsystem::GetOverlayMaterialNames() const
+{
+	TArray<FName> OutKeys;
+	RegisteredOverlayMaterials.GetKeys(OutKeys);
+
+#if WITH_EDITOR
+	OutKeys.Append(GetDefault<UCameraCalibrationSettings>()->GetDefaultCalibrationOverlayNames());
+#endif
+	return OutKeys;
+}
+
 TSubclassOf<UCameraCalibrationStep> UCameraCalibrationSubsystem::GetCameraCalibrationStep(FName Name) const
 {
 	if (CameraCalibrationStepsMap.Contains(Name))
diff --git a/Engine/Plugins/VirtualProduction/CameraCalibrationCore/Source/CameraCalibrationCore/Public/CameraCalibrationSettings.h b/Engine/Plugins/VirtualProduction/CameraCalibrationCore/Source/CameraCalibrationCore/Public/CameraCalibrationSettings.h
--- a/Engine/Plugins/VirtualProduction/CameraCalibrationCore/Source/CameraCalibrationCore/Public/CameraCalibrationSettings.h
+++ b/Engine/Plugins/VirtualProduction/CameraCalibrationCore/Source/CameraCalibrationCore/Public/CameraCalibrationSettings.h
@@ -36,6 +36,12 @@ public:
 #if WITH_EDITOR
 	virtual void PostEditChangeChainProperty(struct FPropertyChangedChainEvent& PropertyChangedEvent) override;
 
+	/** Get the list of default overlay names */
+	TArray<FName> GetDefaultCalibrationOverlayNames() const;
+
+	/** Get the default MaterialInterface associated with the input overlay name */
+	UMaterialInterface* GetDefaultCalibrationOverlayMaterial(const FName OverlayName) const;
+
 	/** Gets a multicast delegate which is called whenever the displacement map resolution project setting changes */
 	FOnDisplacementMapResolutionChanged& OnDisplacementMapResolutionChanged();
 
@@ -103,6 +109,12 @@ private:
 	/** Map of Lens Distortion Model Handler classes to the default lens distortion post-process material used by that class */
 	UPROPERTY(config)
 	TMap<TSubclassOf<ULensDistortionModelHandlerBase>, TSoftObjectPtr<UMaterialInterface>> DefaultDistortionMaterials;
+
+#if WITH_EDITORONLY_DATA
+	/** Map of overlay names to default overlay materials */
+	UPROPERTY(config, EditAnywhere, Category = "Overlays")
+	TMap<FName, TSoftObjectPtr<UMaterialInterface>> DefaultCalibrationOverlayMaterials;
+#endif
 };
 
 /**
@@ -136,7 +148,7 @@ struct FLensDataCategoryEditorColor
 			return FColor::Black;
 		}
 	}
-	
+
 	UPROPERTY(EditAnywhere, Category = "Settings")
 	FColor Focus = FColor::Red;
 
@@ -169,14 +181,15 @@ class CAMERACALIBRATIONCORE_API UCameraCalibrationEditorSettings : public UDevel
 	GENERATED_BODY()
 
 public:
+	UCameraCalibrationEditorSettings();
 
-		//~ Begin UDevelopperSettings interface
-		virtual FName GetCategoryName() const;
+	//~ Begin UDeveloperSettings interface
+	virtual FName GetCategoryName() const;
 #if WITH_EDITOR
-		virtual FText GetSectionText() const override;
-		virtual FName GetSectionName() const override;
+	virtual FText GetSectionText() const override;
+	virtual FName GetSectionName() const override;
 #endif
-		//~ End UDevelopperSettings interface
+	//~ End UDeveloperSettings interface
 
 public:
 
diff --git a/Engine/Plugins/VirtualProduction/CameraCalibrationCore/Source/CameraCalibrationCore/Public/CameraCalibrationSubsystem.h b/Engine/Plugins/VirtualProduction/CameraCalibrationCore/Source/CameraCalibrationCore/Public/CameraCalibrationSubsystem.h
--- a/Engine/Plugins/VirtualProduction/CameraCalibrationCore/Source/CameraCalibrationCore/Public/CameraCalibrationSubsystem.h
+++ b/Engine/Plugins/VirtualProduction/CameraCalibrationCore/Source/CameraCalibrationCore/Public/CameraCalibrationSubsystem.h
@@ -88,6 +88,18 @@ public:
 	UFUNCTION(BlueprintCallable, Category = "Lens Distortion")
 	TArray<FName> GetCameraImageCenterAlgos() const;
 
+	/** Returns the overlay material associated with the input overlay name */
+	UFUNCTION(BlueprintCallable, Category = "Lens Distortion")
+	UMaterialInterface* GetOverlayMaterial(const FName& OverlayName) const;
+
+	/** 
+	 * Returns a list of all overlays known to the subsystem
+	 * This includes the default overlays listed in the camera calibration settings 
+	 * as well as any of overlays that have been registered with this subsystem
+	 */
+	UFUNCTION(BlueprintCallable, Category = "Lens Distortion")
+	TArray<FName> GetOverlayMaterialNames() const;
+
 	/** Returns an array with the names of the available camera calibration steps */
 	UFUNCTION(BlueprintCallable, Category = "Lens Distortion")
 	TArray<FName> GetCameraCalibrationSteps() const;
@@ -114,6 +126,12 @@ public:
 	/** Update the overscanned focal length of the input camera component */
 	void UpdateOverscanFocalLength(UCineCameraComponent* Component, float InFocalLength);
 
+	/** Register a new overlay material name and path that can be queried from camera calibration tools */
+	void RegisterOverlayMaterial(const FName& MaterialName, const FName& MaterialPath);
+
+	/** Unregister an overlay material */
+	void UnregisterOverlayMaterial(const FName& MaterialName);
+
 	/** 
 	 * Get the original focal length of the input camera component, if it exists in the subsystems map. 
 	 * Returns false and does not update the output parameter if the input camera is not in the map. 
@@ -141,6 +159,9 @@ private:
 	UPROPERTY(Transient)
 	TMap<FName, TSubclassOf<UCameraCalibrationStep>> CameraCalibrationStepsMap;
 
+	/** Map of overlay names to overlay materials */
+	TMap<FName, TSoftObjectPtr<UMaterialInterface>> RegisteredOverlayMaterials;
+
 	/** Map of actor components to the authoritative lens model that should be used with that component */
 	TMap<FObjectKey, TSubclassOf<ULensModel>> ComponentsWithAuthoritativeModels;
 
