- Disassociate lens distortion model handlers from asset user data
- Add cameracomponent -> distortion handler map to camera calibration subsystem
- Add distortion source picker to composure to choose from multiple available handlers

diff --git a/Engine/Plugins/Animation/LiveLink/Source/LiveLinkComponents/Private/LiveLinkComponentController.cpp b/Engine/Plugins/Animation/LiveLink/Source/LiveLinkComponents/Private/LiveLinkComponentController.cpp
--- a/Engine/Plugins/Animation/LiveLink/Source/LiveLinkComponents/Private/LiveLinkComponentController.cpp
+++ b/Engine/Plugins/Animation/LiveLink/Source/LiveLinkComponents/Private/LiveLinkComponentController.cpp
@@ -13,6 +13,7 @@
 #include "UObject/UObjectIterator.h"
 
 #if WITH_EDITOR
+#include "Editor.h"
 #include "Framework/Notifications/NotificationManager.h"
 #include "Kismet2/ComponentEditorUtils.h"
 #include "Widgets/Notifications/SNotificationList.h"
@@ -29,18 +30,23 @@ ULiveLinkComponentController::ULiveLinkComponentController()
 	PrimaryComponentTick.bStartWithTickEnabled = true;
 	PrimaryComponentTick.TickGroup = ETickingGroup::TG_PrePhysics;
 	bTickInEditor = true;
+
+#if WITH_EDITOR
+	FEditorDelegates::EndPIE.AddUObject(this, &ULiveLinkComponentController::OnEndPIE);
+#endif //WITH_EDITOR
+}
+
+ULiveLinkComponentController::~ULiveLinkComponentController()
+{
+#if WITH_EDITOR
+	FEditorDelegates::EndPIE.RemoveAll(this);
+#endif //WITH_EDITOR
 }
 
 void ULiveLinkComponentController::OnSubjectRoleChanged()
 {
 	//Whenever the subject role is changed, we start from clean controller map. Cleanup the ones currently active
-	for (TPair<TSubclassOf<ULiveLinkRole>, ULiveLinkControllerBase*>& ControllerPair : ControllerMap)
-	{
-		if (ControllerPair.Value)
-		{
-			ControllerPair.Value->Cleanup();
-		}
-	}
+	CleanupControllersInMap();
 
 	if (SubjectRepresentation.Role == nullptr)
 	{
@@ -141,16 +147,18 @@ void ULiveLinkComponentController::OnRegister()
 	bIsDirty = true;
 }
 
+#if WITH_EDITOR
+void ULiveLinkComponentController::OnEndPIE(bool bIsSimulating)
+{
+	// Cleanup each controller when PIE session is ending
+	CleanupControllersInMap();
+}
+#endif //WITH_EDITOR
+
 void ULiveLinkComponentController::DestroyComponent(bool bPromoteChildren /*= false*/)
 {
-	//Whenever the subject role is changed, we start from clean controller map. Cleanup the ones currently active
-	for (TPair<TSubclassOf<ULiveLinkRole>, ULiveLinkControllerBase*>& ControllerPair : ControllerMap)
-	{
-		if (ControllerPair.Value)
-		{
-			ControllerPair.Value->Cleanup();
-		}
-	}
+	// Cleanup each controller before this component is destroyed
+	CleanupControllersInMap();
 
 	Super::DestroyComponent(bPromoteChildren);
 }
@@ -187,6 +195,7 @@ void ULiveLinkComponentController::TickComponent(float DeltaTime, ELevelTick Tic
 			if (bIsDirty)
 			{
 				Controller->SetAttachedComponent(ComponentToControl.GetComponent(GetOwner()));
+				Controller->SetSelectedSubject(SubjectRepresentation);
 				Controller->OnEvaluateRegistered();
 			}
 			
@@ -335,4 +344,16 @@ TSubclassOf<ULiveLinkControllerBase> ULiveLinkComponentController::GetController
 	return SelectedControllerClass;
 }
 
+void ULiveLinkComponentController::CleanupControllersInMap()
+{
+	//Cleanup the currently active controllers in the map
+	for (TPair<TSubclassOf<ULiveLinkRole>, ULiveLinkControllerBase*>& ControllerPair : ControllerMap)
+	{
+		if (ControllerPair.Value)
+		{
+			ControllerPair.Value->Cleanup();
+		}
+	}
+}
+
 #undef LOCTEXT_NAMESPACE
\ No newline at end of file
diff --git a/Engine/Plugins/Animation/LiveLink/Source/LiveLinkComponents/Private/LiveLinkControllerBase.cpp b/Engine/Plugins/Animation/LiveLink/Source/LiveLinkComponents/Private/LiveLinkControllerBase.cpp
--- a/Engine/Plugins/Animation/LiveLink/Source/LiveLinkComponents/Private/LiveLinkControllerBase.cpp
+++ b/Engine/Plugins/Animation/LiveLink/Source/LiveLinkComponents/Private/LiveLinkControllerBase.cpp
@@ -9,6 +9,11 @@ void ULiveLinkControllerBase::SetAttachedComponent(UActorComponent* ActorCompone
 	AttachedComponent = ActorComponent;
 }
 
+void ULiveLinkControllerBase::SetSelectedSubject(FLiveLinkSubjectRepresentation LiveLinkSubject)
+{
+	SelectedSubject = LiveLinkSubject;
+}
+
 AActor* ULiveLinkControllerBase::GetOuterActor() const
 {
 	for (UObject* Outer = GetOuter(); Outer != nullptr; Outer = Outer->GetOuter())
diff --git a/Engine/Plugins/Animation/LiveLink/Source/LiveLinkComponents/Public/LiveLinkComponentController.h b/Engine/Plugins/Animation/LiveLink/Source/LiveLinkComponents/Public/LiveLinkComponentController.h
--- a/Engine/Plugins/Animation/LiveLink/Source/LiveLinkComponents/Public/LiveLinkComponentController.h
+++ b/Engine/Plugins/Animation/LiveLink/Source/LiveLinkComponents/Public/LiveLinkComponentController.h
@@ -17,6 +17,7 @@ class LIVELINKCOMPONENTS_API ULiveLinkComponentController : public UActorCompone
 	
 public:
 	ULiveLinkComponentController();
+	~ULiveLinkComponentController();
 
 public:
 	UPROPERTY(BlueprintReadOnly, EditAnywhere, Category="LiveLink")
@@ -70,6 +71,11 @@ public:
 	/** Used to notify that the subject role has changed. Mainly from Customization or C++ modification to the subject's Role */
 	void OnSubjectRoleChanged();
 
+#if WITH_EDITOR
+	/** Used to cleanup controllers when exiting PIE */
+	void OnEndPIE(bool bIsSimulating);
+#endif //WITH_EDITOR
+
 	//~ Begin UActorComponent interface
 	virtual void OnRegister() override;
 	virtual void DestroyComponent(bool bPromoteChildren = false) override;
@@ -89,6 +95,9 @@ protected:
 	TArray<TSubclassOf<ULiveLinkRole>> GetSelectedRoleHierarchyClasses(const TSubclassOf<ULiveLinkRole> InCurrentRoleClass) const;
 	TSubclassOf<ULiveLinkControllerBase> GetControllerClassForRoleClass(const TSubclassOf<ULiveLinkRole> RoleClass) const;
 
+	/** Loops through the controller map and calls Cleanup() on each entry */
+	void CleanupControllersInMap();
+
 #if WITH_EDITOR
 	/** Called during loading to convert old data to new scheme. */
 	void ConvertOldControllerSystem();
diff --git a/Engine/Plugins/Animation/LiveLink/Source/LiveLinkComponents/Public/LiveLinkControllerBase.h b/Engine/Plugins/Animation/LiveLink/Source/LiveLinkComponents/Public/LiveLinkControllerBase.h
--- a/Engine/Plugins/Animation/LiveLink/Source/LiveLinkComponents/Public/LiveLinkControllerBase.h
+++ b/Engine/Plugins/Animation/LiveLink/Source/LiveLinkComponents/Public/LiveLinkControllerBase.h
@@ -49,6 +49,12 @@ public:
 	 */
 	virtual void SetAttachedComponent(UActorComponent* ActorComponent);
 
+	/**
+	 * Sets the live link subject from which this controller is receiving data
+	 */
+	virtual void SetSelectedSubject(FLiveLinkSubjectRepresentation LiveLinkSubject);
+
+
 	/**
 	 * Cleanup controller state before getting removed
 	 */
@@ -73,5 +79,7 @@ public:
 
 protected:
 	TWeakObjectPtr<UActorComponent> AttachedComponent;
+
+	FLiveLinkSubjectRepresentation SelectedSubject;
 };
 
diff --git a/Engine/Plugins/Compositing/Composure/Source/Composure/Classes/CompositingCaptureBase.h b/Engine/Plugins/Compositing/Composure/Source/Composure/Classes/CompositingCaptureBase.h
--- a/Engine/Plugins/Compositing/Composure/Source/Composure/Classes/CompositingCaptureBase.h
+++ b/Engine/Plugins/Compositing/Composure/Source/Composure/Classes/CompositingCaptureBase.h
@@ -4,6 +4,7 @@
 
 #include "CompositingElement.h"
 
+#include "CameraCalibrationTypes.h"
 #include "LensDistortionModelHandlerBase.h"
 
 #include "CompositingCaptureBase.generated.h"
@@ -28,12 +29,16 @@ protected:
 	UPROPERTY(EditAnywhere, BlueprintReadWrite, Category = "Composure|LensDistortion")
 	bool bApplyDistortion = false;
 
-	/** Pointer to the Lens Distortion Data Handler that belongs to the CameraComponent of the TargetCameraActor */
-	UPROPERTY(EditAnywhere, BlueprintReadWrite, Category = "Composure|LensDistortion", Transient)
-	ULensDistortionModelHandlerBase* LensDistortionHandler = nullptr;
+	/** Structure used to query the camera calibration subsystem for a lens distortion model handler */
+	UPROPERTY(EditAnywhere, BlueprintReadWrite, Category = "Composure|LensDistortion")
+	FDistortionHandlerPicker DistortionSource;
+
+	/** Value used to augment the FOV of the scene capture to produce a CG image with enough data to distort */
+	UPROPERTY(BlueprintReadOnly, Category = "Composure|LensDistortion")
+	float OverscanFactor = 1.0f;
 
 	/** Cached distortion MID produced by the Lens Distortion Handler, used to clean up the post-process materials in the case that the the MID changes */
-	UPROPERTY(Transient)
+	UPROPERTY()
 	UMaterialInstanceDynamic* LastDistortionMID = nullptr;
 
 public:
diff --git a/Engine/Plugins/Compositing/Composure/Source/Composure/Classes/ComposureBlueprintLibrary.h b/Engine/Plugins/Compositing/Composure/Source/Composure/Classes/ComposureBlueprintLibrary.h
--- a/Engine/Plugins/Compositing/Composure/Source/Composure/Classes/ComposureBlueprintLibrary.h
+++ b/Engine/Plugins/Compositing/Composure/Source/Composure/Classes/ComposureBlueprintLibrary.h
@@ -69,7 +69,7 @@ class COMPOSURE_API UComposureBlueprintLibrary : public UBlueprintFunctionLibrar
 	static void GetPlayerDisplayGamma(const APlayerCameraManager* PlayerCameraManager, float& DisplayGamma);
 
 	UFUNCTION(BlueprintCallable, Category = "Composure")
-	static void CopyCameraSettingsToSceneCapture(UCameraComponent* SrcCamera, USceneCaptureComponent2D* DstCaptureComponent, bool bApplyDistortion = false);
+	static void CopyCameraSettingsToSceneCapture(UCameraComponent* SrcCamera, USceneCaptureComponent2D* DstCaptureComponent, float OverscanFactor = 1.0f);
 
 	/**
 	 * Create a new Composure in the level without any parenting relationship.
diff --git a/Engine/Plugins/Compositing/Composure/Source/Composure/Private/CompositingCaptureBase.cpp b/Engine/Plugins/Compositing/Composure/Source/Composure/Private/CompositingCaptureBase.cpp
--- a/Engine/Plugins/Compositing/Composure/Source/Composure/Private/CompositingCaptureBase.cpp
+++ b/Engine/Plugins/Compositing/Composure/Source/Composure/Private/CompositingCaptureBase.cpp
@@ -5,6 +5,7 @@
 #include "CameraCalibrationSubsystem.h"
 #include "CineCameraComponent.h"
 #include "Components/SceneCaptureComponent2D.h"
+#include "Engine/Engine.h"
 
 ACompositingCaptureBase::ACompositingCaptureBase()
 {
@@ -24,14 +25,17 @@ void ACompositingCaptureBase::UpdateDistortion()
 	{
 		return;
 	}
-
+	
  	if (UCineCameraComponent* const CineCameraComponent = Cast<UCineCameraComponent>(TargetCamera->GetCameraComponent()))
 	{
-		// Get the CineCameraComponent's LensDistortionDataHandler, if it exists, or create a new one if it does not
+		DistortionSource.TargetCameraComponent = CineCameraComponent;
+
+		// Query the camera calibration subsystem for a handler associated with the TargetCamera and matching the user selection
+		ULensDistortionModelHandlerBase* LensDistortionHandler = nullptr;
 		UCameraCalibrationSubsystem* SubSystem = GEngine->GetEngineSubsystem<UCameraCalibrationSubsystem>();
 		if (SubSystem)
 		{
-			LensDistortionHandler = SubSystem->GetDistortionModelHandler(CineCameraComponent);
+			LensDistortionHandler = SubSystem->FindDistortionModelHandler(DistortionSource);
 		}
 
 		if (LensDistortionHandler)
@@ -54,14 +58,26 @@ void ACompositingCaptureBase::UpdateDistortion()
 			{
 				if (bApplyDistortion)
 				{
+					OverscanFactor = LensDistortionHandler->GetOverscanFactor();
 					SceneCaptureComponent2D->AddOrUpdateBlendable(NewDistortionMID);
 				}
 				else
 				{
+					OverscanFactor = 1.0f;
 					SceneCaptureComponent2D->RemoveBlendable(NewDistortionMID);
 				}
 			}
 		}
+		else
+		{
+			OverscanFactor = 1.0f;
+
+			if (SceneCaptureComponent2D)
+			{
+				SceneCaptureComponent2D->RemoveBlendable(LastDistortionMID);
+			}
+			LastDistortionMID = nullptr;
+		}
 	}
 }
 
@@ -72,6 +88,27 @@ void ACompositingCaptureBase::PostEditChangeProperty(struct FPropertyChangedEven
 
 	if ((PropertyName == GET_MEMBER_NAME_CHECKED(ACompositingCaptureBase, TargetCameraActor)))
 	{
+		if (TargetCameraActor)
+		{
+			if (UCineCameraComponent* const CineCameraComponent = Cast<UCineCameraComponent>(TargetCameraActor->GetCameraComponent()))
+			{
+				DistortionSource.TargetCameraComponent = CineCameraComponent;
+
+				ULensDistortionModelHandlerBase* LensDistortionHandler = nullptr;
+				UCameraCalibrationSubsystem* SubSystem = GEngine->GetEngineSubsystem<UCameraCalibrationSubsystem>();
+				if (SubSystem)
+				{
+					LensDistortionHandler = SubSystem->FindDistortionModelHandler(DistortionSource);
+				}
+
+				if (LensDistortionHandler == nullptr)
+				{
+					DistortionSource.DistortionProducerID.Invalidate();
+					DistortionSource.HandlerDisplayName.Empty();
+				}
+			}
+		}
+
 		// If there is no target camera, remove the last distortion post-process MID from the scene capture
 		if (TargetCameraActor == nullptr)
 		{
diff --git a/Engine/Plugins/Compositing/Composure/Source/Composure/Private/ComposureBlueprintLibrary.cpp b/Engine/Plugins/Compositing/Composure/Source/Composure/Private/ComposureBlueprintLibrary.cpp
--- a/Engine/Plugins/Compositing/Composure/Source/Composure/Private/ComposureBlueprintLibrary.cpp
+++ b/Engine/Plugins/Compositing/Composure/Source/Composure/Private/ComposureBlueprintLibrary.cpp
@@ -72,41 +72,34 @@ void UComposureBlueprintLibrary::GetPlayerDisplayGamma(const APlayerCameraManage
 	DisplayGamma = SceneViewport ? SceneViewport->GetDisplayGamma() : 0.0;
 }
 
-void UComposureBlueprintLibrary::CopyCameraSettingsToSceneCapture(UCameraComponent* Src, USceneCaptureComponent2D* Dst, bool bApplyDistortion)
+void UComposureBlueprintLibrary::CopyCameraSettingsToSceneCapture(UCameraComponent* Src, USceneCaptureComponent2D* Dst, float OverscanFactor)
 {
 	if (Src && Dst)
 	{
 		Dst->SetWorldLocationAndRotation(Src->GetComponentLocation(), Src->GetComponentRotation());
-		Dst->FOVAngle = Src->FieldOfView;
 
 		FMinimalViewInfo CameraViewInfo;
 		Src->GetCameraView(/*DeltaTime =*/0.0f, CameraViewInfo);
 
-		// If distortion should be applied, get the overscan factor from the Src's lens distortion data handler, and use it to augment the Dst's FOV
-		if (bApplyDistortion)
-		{
-			if (UCineCameraComponent* SrcCineCameraComponent = Cast<UCineCameraComponent>(Src))
-			{ 			
-				UCameraCalibrationSubsystem* SubSystem = GEngine->GetEngineSubsystem<UCameraCalibrationSubsystem>();
-				if (SubSystem)
-				{
-					if (ULensDistortionModelHandlerBase* LensDistortionHandler = SubSystem->GetDistortionModelHandler(SrcCineCameraComponent))
-					{
-						const float OverscanFactor = LensDistortionHandler->GetOverscanFactor();
-						if (SrcCineCameraComponent->CurrentFocalLength <= 0.0f)
-						{
-							Dst->FOVAngle = 0.0f;
-						}
-						else
-						{
-							const float OverscanSensorWidth = (SrcCineCameraComponent->Filmback.SensorWidth * 0.5f) * OverscanFactor;
-							Dst->FOVAngle = FMath::RadiansToDegrees(2.0f * FMath::Atan(OverscanSensorWidth / SrcCineCameraComponent->CurrentFocalLength));
-						}
-					}
-
-				}
+		// Use the input overscan factor to augment the destination component's FOV angle
+		// Note: The math relies on the filmback and focal length of the input camera component, which necessitates that it be a CineCameraComponet
+		if (UCineCameraComponent* SrcCineCameraComponent = Cast<UCineCameraComponent>(Src))
+		{ 			
+			// Guard against divide-by-zero
+			if (SrcCineCameraComponent->CurrentFocalLength <= 0.0f)
+			{
+				Dst->FOVAngle = 0.0f;
+			}
+			else
+			{
+				const float OverscanSensorWidth = SrcCineCameraComponent->Filmback.SensorWidth * OverscanFactor;
+				Dst->FOVAngle = FMath::RadiansToDegrees(2.0f * FMath::Atan(OverscanSensorWidth / (2.0f * SrcCineCameraComponent->CurrentFocalLength)));
 			}
 		}
+		else
+		{
+			Dst->FOVAngle = Src->FieldOfView;
+		}
 
 		const FPostProcessSettings& SrcPPSettings = CameraViewInfo.PostProcessSettings;
 		FPostProcessSettings& DstPPSettings = Dst->PostProcessSettings;
diff --git a/Engine/Plugins/VirtualProduction/CameraCalibration/Source/CameraCalibration/Private/CameraCalibrationSubsystem.cpp b/Engine/Plugins/VirtualProduction/CameraCalibration/Source/CameraCalibration/Private/CameraCalibrationSubsystem.cpp
--- a/Engine/Plugins/VirtualProduction/CameraCalibration/Source/CameraCalibration/Private/CameraCalibrationSubsystem.cpp
+++ b/Engine/Plugins/VirtualProduction/CameraCalibration/Source/CameraCalibration/Private/CameraCalibrationSubsystem.cpp
@@ -3,14 +3,13 @@
 #include "CameraCalibrationSubsystem.h"
 
 #include "CameraCalibrationLog.h"
-#include "Components/ActorComponent.h"
+#include "CineCameraComponent.h"
 #include "Engine/TimecodeProvider.h"
-#include "Interfaces/Interface_AssetUserData.h"
 #include "UObject/UObjectIterator.h"
 
 #if WITH_EDITOR
 #include "Editor.h"
-#endif
+#endif // WITH_EDITOR
 
 #define LOCTEXT_NAMESPACE "CameraCalibrationSubsystem"
 
@@ -43,88 +42,83 @@ ULensFile* UCameraCalibrationSubsystem::GetLensFile(const FLensFilePicker& Picke
 	return ReturnedLens;
 }
 
-ULensDistortionModelHandlerBase* UCameraCalibrationSubsystem::GetDistortionModelHandler(UActorComponent* Component)
+TArray<ULensDistortionModelHandlerBase*> UCameraCalibrationSubsystem::GetDistortionModelHandlers(UCineCameraComponent* Component)
 {
-	if (IInterface_AssetUserData* AssetUserData = Cast<IInterface_AssetUserData>(Component))
+	TArray<ULensDistortionModelHandlerBase*> Handlers;
+	LensDistortionHandlerMap.MultiFind(Component, Handlers);
+
+	return Handlers;
+}
+
+ULensDistortionModelHandlerBase* UCameraCalibrationSubsystem::FindDistortionModelHandler(FDistortionHandlerPicker& DistortionHandlerPicker, bool bUpdatePicker) const
+{
+	TArray<ULensDistortionModelHandlerBase*> Handlers;
+	LensDistortionHandlerMap.MultiFind(DistortionHandlerPicker.TargetCameraComponent, Handlers);
+
+	// Look through the available handlers to find the one driven by the distortion source's producer
+	for (ULensDistortionModelHandlerBase* Handler : Handlers)
 	{
-		return Cast<ULensDistortionModelHandlerBase>(AssetUserData->GetAssetUserDataOfClass(ULensDistortionModelHandlerBase::StaticClass()));
+		if (Handler->GetDistortionProducerID() == DistortionHandlerPicker.DistortionProducerID)
+		{
+			// Reassign the input handler picker's display name to that of the found handler
+			if (bUpdatePicker)
+			{
+				DistortionHandlerPicker.HandlerDisplayName = Handler->GetDisplayName();
+			}
+			return Handler;
+		}
+	}
+
+	// Look through the available handlers to find one with the same name as the input handler picker's display name
+	for (ULensDistortionModelHandlerBase* Handler : Handlers)
+	{
+		if (Handler->GetDisplayName() == DistortionHandlerPicker.HandlerDisplayName)
+		{
+			// Reassign the input handler picker's producer to that of the found handler
+			if (bUpdatePicker)
+			{
+				DistortionHandlerPicker.DistortionProducerID = Handler->GetDistortionProducerID();
+			}
+			return Handler;
+		}
 	}
 
 	return nullptr;
 }
 
-ULensDistortionModelHandlerBase* UCameraCalibrationSubsystem::FindOrCreateDistortionModelHandler(UActorComponent* Component, TSubclassOf<ULensModel> LensModelClass, EHandlerOverrideMode OverrideMode)
+ULensDistortionModelHandlerBase* UCameraCalibrationSubsystem::FindOrCreateDistortionModelHandler(FDistortionHandlerPicker& DistortionHandlerPicker, const TSubclassOf<ULensModel> LensModelClass)
 {
 	if (LensModelClass == nullptr)
 	{
 		return nullptr;
 	}
 
-	// Get the current model handler on the input component (if it exists)
-	if (ULensDistortionModelHandlerBase* Handler = GetDistortionModelHandler(Component))
+	// Attempt to find a handler associated with the input distortion source
+	if (ULensDistortionModelHandlerBase* Handler = FindDistortionModelHandler(DistortionHandlerPicker, false))
 	{
 		// If the existing handler supports the input model, simply return it
 		if (Handler->IsModelSupported(LensModelClass))
 		{
+			// The display name may have changed, even if the handler already exists, so update the display name of the existing handler to match the distortion source
+			Handler->SetDisplayName(DistortionHandlerPicker.HandlerDisplayName);
 			return Handler;
 		}
-		else if (OverrideMode == EHandlerOverrideMode::NoOverride)
-		{
-			return nullptr;
-		}
-
-		UActorComponent* ComponentToCheck = Component;
-#if WITH_EDITOR
-		// Actors are duplicated in PIE, so their components will not appear in the subsystem's set of components with authoritative models
-		// Therefore, the component to check should be the editor world counterpart to the input component
-		if (Component->GetWorld() && (Component->GetWorld()->WorldType == EWorldType::PIE))
-		{
-			if (AActor* PIEOwner = EditorUtilities::GetEditorWorldCounterpartActor(Component->GetOwner()))
-			{
-				ComponentToCheck = PIEOwner->GetComponentByClass(Component->GetClass());
-			}
-		}
-#endif
-
-		if (OverrideMode == EHandlerOverrideMode::SoftOverride)
+		else
 		{
-			// If the input component already has an authoritative model, do not override the existing handler
-			if (ComponentsWithAuthoritativeModels.Contains(FObjectKey(ComponentToCheck))
-				&& (ComponentsWithAuthoritativeModels[ComponentToCheck] != LensModelClass))
-			{
-				return nullptr;
-			}
-			// Mark the input component as having an authoritative model and remove the existing handler
-			else
-			{
-				ComponentsWithAuthoritativeModels.Add(FObjectKey(ComponentToCheck), LensModelClass);
-				Component->RemoveUserDataOfClass(ULensDistortionModelHandlerBase::StaticClass());
-			}
-		}
-		// Mark the input component as having an authoritative model and remove the existing handler
-		else if (OverrideMode == EHandlerOverrideMode::ForceOverride)
-		{
-			// Override the authoritative model for this component, if it already has one
-			if (ComponentsWithAuthoritativeModels.Contains(FObjectKey(ComponentToCheck))
-				&& (ComponentsWithAuthoritativeModels[ComponentToCheck] != LensModelClass))
-			{
-				ComponentsWithAuthoritativeModels[ComponentToCheck] = LensModelClass;
-			}
-			else
-			{
-				ComponentsWithAuthoritativeModels.Add(FObjectKey(ComponentToCheck), LensModelClass);
-			}
-
-			Component->RemoveUserDataOfClass(ULensDistortionModelHandlerBase::StaticClass());
+			// If the input distortion source has an existing handler, but model does not match the input model, remove the old handler from the map
+			LensDistortionHandlerMap.Remove(DistortionHandlerPicker.TargetCameraComponent, Handler);
 		}
 	}
 
+	// If no handler exists for the input distortion source, create a new one
 	const TSubclassOf<ULensDistortionModelHandlerBase> HandlerClass = ULensModel::GetHandlerClass(LensModelClass);
 	ULensDistortionModelHandlerBase* NewHandler = nullptr;
 	if(HandlerClass)
 	{
-		NewHandler = NewObject<ULensDistortionModelHandlerBase>(Component, HandlerClass);
-    	Component->AddAssetUserData(NewHandler);
+		NewHandler = NewObject<ULensDistortionModelHandlerBase>(DistortionHandlerPicker.TargetCameraComponent, HandlerClass);
+		NewHandler->SetDistortionProducerID(DistortionHandlerPicker.DistortionProducerID);
+		NewHandler->SetDisplayName(DistortionHandlerPicker.HandlerDisplayName);
+		LensDistortionHandlerMap.Add(FObjectKey(DistortionHandlerPicker.TargetCameraComponent), NewHandler);
 	}
 	else
 	{
@@ -134,6 +128,12 @@ ULensDistortionModelHandlerBase* UCameraCalibrationSubsystem::FindOrCreateDistor
 	return NewHandler;
 }
 
+void UCameraCalibrationSubsystem::UnregisterDistortionModelHandler(UCineCameraComponent* Component, ULensDistortionModelHandlerBase* Handler)
+{
+	LensDistortionHandlerMap.Remove(Component, Handler);
+}
+
+
 void UCameraCalibrationSubsystem::RegisterDistortionModel(TSubclassOf<ULensModel> LensModel)
 {
 	LensModelMap.Add(LensModel->GetDefaultObject<ULensModel>()->GetModelName(), LensModel);
@@ -144,7 +144,7 @@ void UCameraCalibrationSubsystem::UnregisterDistortionModel(TSubclassOf<ULensMod
 	LensModelMap.Remove(LensModel->GetDefaultObject<ULensModel>()->GetModelName());
 }
 
-TSubclassOf<ULensModel> UCameraCalibrationSubsystem::GetRegisteredLensModel(FName ModelName)
+TSubclassOf<ULensModel> UCameraCalibrationSubsystem::GetRegisteredLensModel(FName ModelName) const
 {
 	if (LensModelMap.Contains(ModelName))
 	{
diff --git a/Engine/Plugins/VirtualProduction/CameraCalibration/Source/CameraCalibration/Public/CameraCalibrationSubsystem.h b/Engine/Plugins/VirtualProduction/CameraCalibration/Source/CameraCalibration/Public/CameraCalibrationSubsystem.h
--- a/Engine/Plugins/VirtualProduction/CameraCalibration/Source/CameraCalibration/Public/CameraCalibrationSubsystem.h
+++ b/Engine/Plugins/VirtualProduction/CameraCalibration/Source/CameraCalibration/Public/CameraCalibrationSubsystem.h
@@ -4,32 +4,14 @@
 
 #include "Subsystems/EngineSubsystem.h"
 
+#include "CameraCalibrationTypes.h"
 #include "LensDistortionModelHandlerBase.h"
 #include "LensFile.h"
 #include "UObject/ObjectKey.h"
 
 #include "CameraCalibrationSubsystem.generated.h"
 
-class UActorComponent;
-
-UENUM(BlueprintType)
-enum class EHandlerOverrideMode : uint8
-{
-	/** If a component already has a lens handler that supports a different lens model, do not override the existing handler */
-	NoOverride,
-
-	/** 
-	 * If a component already has a lens handler that supports a different lens model, but that model is not authoritative,
-	 * remove the existing handler, create a new handler for the new model, and mark that new model as the authoritative one for that component
-	 */
-	SoftOverride,
-
-	/**
-	 * If a component already has a lens handler that supports a different lens model, regardless of any previous authoritative models,
-	 * remove the existing handler, create a new handler for the new model, and mark that new model as the authoritative one for that component
-	 */
-	 ForceOverride
-};
+class UCineCameraComponent;
 
 /**
  * Camera Calibration subsystem
@@ -53,21 +35,28 @@ public:
 	UFUNCTION(BlueprintCallable, Category = "Lens Distortion")
 	ULensFile* GetLensFile(const FLensFilePicker& Picker) const;
 
-	/** Get the lens distortion model handler owned by the input actor component (returns the first instance if there are multiple) */
+	/** Return all handlers associated with the input camera component */
 	UFUNCTION(BlueprintCallable, Category = "Lens Distortion")
-	ULensDistortionModelHandlerBase* GetDistortionModelHandler(UActorComponent* Component);
+	TArray<ULensDistortionModelHandlerBase*> GetDistortionModelHandlers(UCineCameraComponent* Component);
 
 	/** 
-	 * Get the lens distortion model handler owned by the input actor component that supports the input lens model.
-	 * If no such handler exists, creates a new instance of one that supports the input lens model and gives ownership of the handler to the input component. 
-	 * If the input component already has a model handler that does not support the input model, that handler is removed.
+	 * Return the handler associated with the input distortion source, if one exists 
+	 * If bUpdatePicker is true, the input picker reference will be updated so that its properties match those of the found handler
 	 */
 	UFUNCTION(BlueprintCallable, Category = "Lens Distortion")
-	ULensDistortionModelHandlerBase* FindOrCreateDistortionModelHandler(UActorComponent* Component, TSubclassOf<ULensModel> LensModelClass, EHandlerOverrideMode OverrideMode = EHandlerOverrideMode::NoOverride);
+	ULensDistortionModelHandlerBase* FindDistortionModelHandler(FDistortionHandlerPicker& DistortionHandlerPicker, bool bUpdatePicker = true) const;
+
+	/** Return the handler associated with the input distortion source, if one exists that also matches the input model. If none exist, create a new handler and return it. */
+	UFUNCTION(BlueprintCallable, Category = "Lens Distortion")
+	ULensDistortionModelHandlerBase* FindOrCreateDistortionModelHandler(FDistortionHandlerPicker& DistortionHandlerPicker, const TSubclassOf<ULensModel> LensModelClass);
+
+	/** Disassociate the input handler from the input camera component in the subsystem's handler map */
+	UFUNCTION(BlueprintCallable, Category = "Lens Distortion")
+	void UnregisterDistortionModelHandler(UCineCameraComponent* Component, ULensDistortionModelHandlerBase* Handler);
 
 	/** Return the ULensModel subclass that was registered with the input model name */
 	UFUNCTION(BlueprintCallable, Category = "Lens Distortion")
-	TSubclassOf<ULensModel> GetRegisteredLensModel(FName ModelName);
+	TSubclassOf<ULensModel> GetRegisteredLensModel(FName ModelName) const;
 
 public:
 	//~ Begin USubsystem interface
@@ -91,4 +80,8 @@ private:
 
 	/** Map of actor components to the authoritative lens model that should be used with that component */
 	TMap<FObjectKey, TSubclassOf<ULensModel>> ComponentsWithAuthoritativeModels;
+
+	/** Map of actor components to the authoritative lens model that should be used with that component */
+	TMultiMap<FObjectKey, ULensDistortionModelHandlerBase*> LensDistortionHandlerMap;
+
 };
diff --git a/Engine/Plugins/VirtualProduction/CameraCalibration/Source/CameraCalibration/Public/CameraCalibrationTypes.h b/Engine/Plugins/VirtualProduction/CameraCalibration/Source/CameraCalibration/Public/CameraCalibrationTypes.h
new file mode 100644
--- /dev/null
+++ b/Engine/Plugins/VirtualProduction/CameraCalibration/Source/CameraCalibration/Public/CameraCalibrationTypes.h
@@ -0,0 +1,29 @@
+// Copyright Epic Games, Inc. All Rights Reserved.
+
+#pragma once
+
+#include "CineCameraComponent.h"
+#include "LensDistortionModelHandlerBase.h"
+#include "Misc/Guid.h"
+
+#include "CameraCalibrationTypes.generated.h"
+
+/** Utility structure for selecting a distortion handler from the camera calibration subsystem */
+USTRUCT(BlueprintType)
+struct CAMERACALIBRATION_API FDistortionHandlerPicker
+{
+	GENERATED_BODY()
+
+public:
+	/** CineCameraComponent with which the desired distortion handler is associated */
+	UPROPERTY(Transient)
+	UCineCameraComponent* TargetCameraComponent = nullptr;
+
+	/** UObject that produces the distortion state for the desired distortion handler */
+	UPROPERTY()
+	FGuid DistortionProducerID;
+
+	/** Display name of the desired distortion handler */
+	UPROPERTY()
+	FString HandlerDisplayName;
+};
\ No newline at end of file
diff --git a/Engine/Plugins/VirtualProduction/CameraCalibration/Source/CameraCalibration/Public/LensDistortionModelHandlerBase.h b/Engine/Plugins/VirtualProduction/CameraCalibration/Source/CameraCalibration/Public/LensDistortionModelHandlerBase.h
--- a/Engine/Plugins/VirtualProduction/CameraCalibration/Source/CameraCalibration/Public/LensDistortionModelHandlerBase.h
+++ b/Engine/Plugins/VirtualProduction/CameraCalibration/Source/CameraCalibration/Public/LensDistortionModelHandlerBase.h
@@ -2,11 +2,10 @@
 
 #pragma once
 
-#include "Engine/AssetUserData.h"
-
 #include "CoreMinimal.h"
 #include "LensData.h"
 #include "Materials/MaterialInstanceDynamic.h"
+#include "Misc/Guid.h"
 #include "Models/LensModel.h"
 
 #include "LensDistortionModelHandlerBase.generated.h"
@@ -32,7 +31,7 @@ public:
 
 /** Asset user data that can be used on Camera Actors to manage lens distortion state and utilities  */
 UCLASS(Abstract)
-class CAMERACALIBRATION_API ULensDistortionModelHandlerBase : public UAssetUserData
+class CAMERACALIBRATION_API ULensDistortionModelHandlerBase : public UObject
 {
 	GENERATED_BODY()
 
@@ -71,6 +70,18 @@ public:
 	/** Get the specified lens model that characterizes the distortion effect */
 	const TSubclassOf<ULensModel>& GetLensModelClass() const { return LensModelClass; };
 
+	/** Get the UObject that produces the distortion state for this handler */
+	FGuid GetDistortionProducerID() const { return DistortionProducerID; }
+
+	/** Set the UObject that produces the distortion state for this handler */
+	void SetDistortionProducerID(const FGuid& InDistortionProducerID) { DistortionProducerID = InDistortionProducerID; }
+
+	/** Get the display name of this lens distortion model handler */
+	FString GetDisplayName() const { return DisplayName; }
+
+	/** Set the display name of this lens distortion model handler */
+	void SetDisplayName(FString InDisplayName) { DisplayName = InDisplayName; }
+
 	/** Get the normalized center of projection of the image, in the range [0.0f, 1.0f] */
 	FVector2D GetPrincipalPoint() const { return CurrentState.PrincipalPoint; }
 
@@ -127,6 +138,10 @@ protected:
 	UPROPERTY(EditAnywhere, BlueprintReadOnly, Category = "Distortion", meta = (ShowOnlyInnerProperties))
 	FLensDistortionState CurrentState;
 
+	/** Display name, used to identify handler in-editor details panels */
+	UPROPERTY(VisibleAnywhere, Category = "Distortion")
+	FString DisplayName;
+
 	/** Computed overscan factor needed to scale the camera's FOV (read-only) */
 	UPROPERTY(VisibleAnywhere, BlueprintReadOnly, Category = "Distortion")
 	float OverscanFactor = 1.0f;
@@ -147,6 +162,10 @@ protected:
 	UPROPERTY(Transient)
 	UTextureRenderTarget2D* DistortionDisplacementMapRT = nullptr;
 
+	/** UObject that is producing the distortion state for this handler */
+	UPROPERTY()
+	FGuid DistortionProducerID;
+
 private:
 	static constexpr uint32 DisplacementMapWidth = 256;
 	static constexpr uint32 DisplacementMapHeight = 256;
diff --git a/Engine/Plugins/VirtualProduction/CameraCalibration/Source/CameraCalibrationEditor/Private/CameraCalibrationEditorModule.cpp b/Engine/Plugins/VirtualProduction/CameraCalibration/Source/CameraCalibrationEditor/Private/CameraCalibrationEditorModule.cpp
--- a/Engine/Plugins/VirtualProduction/CameraCalibration/Source/CameraCalibrationEditor/Private/CameraCalibrationEditorModule.cpp
+++ b/Engine/Plugins/VirtualProduction/CameraCalibration/Source/CameraCalibrationEditor/Private/CameraCalibrationEditorModule.cpp
@@ -6,6 +6,8 @@
 #include "AssetEditor/CameraCalibrationCommands.h"
 #include "AssetToolsModule.h"
 #include "AssetTypeActions/AssetTypeActions_LensFile.h"
+#include "CameraCalibrationTypes.h"
+#include "DistortionHandlerPickerDetailCustomization.h"
 #include "IAssetTools.h"
 #include "IAssetTypeActions.h"
 #include "LensFile.h"
@@ -47,6 +49,9 @@ void FCameraCalibrationEditorModule::StartupModule()
 			true);
 	}
 
+	FPropertyEditorModule& PropertyEditorModule = FModuleManager::Get().LoadModuleChecked<FPropertyEditorModule>("PropertyEditor");
+	PropertyEditorModule.RegisterCustomPropertyTypeLayout(FDistortionHandlerPicker::StaticStruct()->GetFName(), FOnGetPropertyTypeCustomizationInstance::CreateStatic(&FDistortionHandlerPickerDetailCustomization::MakeInstance));
+
 	FCameraCalibrationMenuEntry::Register();
 }
 
@@ -60,6 +65,7 @@ void FCameraCalibrationEditorModule::ShutdownModule()
 
 		FPropertyEditorModule& PropertyModule = FModuleManager::GetModuleChecked<FPropertyEditorModule>("PropertyEditor");
 		PropertyModule.UnregisterCustomClassLayout(ULensFile::StaticClass()->GetFName());
+		PropertyModule.UnregisterCustomPropertyTypeLayout(FDistortionHandlerPicker::StaticStruct()->GetFName());
 
 		// Unregister AssetTypeActions
 		FAssetToolsModule* AssetToolsModule = FModuleManager::GetModulePtr<FAssetToolsModule>("AssetTools");
diff --git a/Engine/Plugins/VirtualProduction/CameraCalibration/Source/CameraCalibrationEditor/Private/DistortionHandlerPickerDetailCustomization.cpp b/Engine/Plugins/VirtualProduction/CameraCalibration/Source/CameraCalibrationEditor/Private/DistortionHandlerPickerDetailCustomization.cpp
new file mode 100644
--- /dev/null
+++ b/Engine/Plugins/VirtualProduction/CameraCalibration/Source/CameraCalibrationEditor/Private/DistortionHandlerPickerDetailCustomization.cpp
@@ -0,0 +1,161 @@
+// Copyright Epic Games, Inc. All Rights Reserved.
+
+#include "DistortionHandlerPickerDetailCustomization.h"
+
+#include "CameraCalibrationSubsystem.h"
+#include "CameraCalibrationTypes.h"
+#include "DetailWidgetRow.h"
+#include "Framework/MultiBox/MultiBoxBuilder.h"
+#include "IPropertyUtilities.h"
+#include "Misc/Guid.h"
+
+
+#define LOCTEXT_NAMESPACE "DistortionHandlerPickerDetailCustomization"
+
+
+TSharedRef<IPropertyTypeCustomization> FDistortionHandlerPickerDetailCustomization::MakeInstance()
+{
+	return MakeShared<FDistortionHandlerPickerDetailCustomization>();
+}
+
+void FDistortionHandlerPickerDetailCustomization::CustomizeHeader(TSharedRef<IPropertyHandle> InPropertyHandle, FDetailWidgetRow& HeaderRow, IPropertyTypeCustomizationUtils& CustomizationUtils)
+{
+	StructPropertyHandle = InPropertyHandle;
+	TSharedPtr<IPropertyUtilities> PropertyUtils = CustomizationUtils.GetPropertyUtilities();
+	
+	StructPropertyHandle->MarkResetToDefaultCustomized();
+
+	TArray<void*> RawData;
+	StructPropertyHandle->AccessRawData(RawData);
+
+	// Multi-select is not supported for this property
+	if (RawData.Num() == 1)
+	{
+		FDistortionHandlerPicker* DistortionHandlerPicker = reinterpret_cast<FDistortionHandlerPicker*>(RawData[0]);
+
+		HeaderRow.NameContent()
+			[
+				StructPropertyHandle->CreatePropertyNameWidget()
+			]
+			.ValueContent()
+			[
+				BuildDistortionHandlerPickerWidget(DistortionHandlerPicker)
+			]
+			.IsEnabled(MakeAttributeLambda([=] { return !InPropertyHandle->IsEditConst() && PropertyUtils->IsPropertyEditingEnabled(); }));
+	}
+}
+
+TSharedRef<SWidget> FDistortionHandlerPickerDetailCustomization::BuildDistortionHandlerPickerWidget(FDistortionHandlerPicker* InDistortionHandlerPicker)
+{
+	TSharedRef<SWidget> ValueWidget =
+		SNew(SHorizontalBox)
+	 	+ SHorizontalBox::Slot()
+ 		.AutoWidth()
+ 		.Padding(FMargin(0.0f, 0.0f, 0.0f, 0.0f))
+ 		[
+ 			SNew(SComboButton)
+ 			.OnGetMenuContent(this, &FDistortionHandlerPickerDetailCustomization::PopulateDistortionHandlerComboButton, InDistortionHandlerPicker)
+	 		.ContentPadding(FMargin(4.0, 2.0))
+			.ButtonContent()
+			[
+				SNew(STextBlock)
+				.Text(this, &FDistortionHandlerPickerDetailCustomization::OnGetButtonText, InDistortionHandlerPicker)
+			]
+		]
+	;
+
+	return ValueWidget;
+}
+
+FText FDistortionHandlerPickerDetailCustomization::OnGetButtonText(FDistortionHandlerPicker* InDistortionHandlerPicker) const
+{
+	if (InDistortionHandlerPicker->HandlerDisplayName.IsEmpty())
+	{
+		return FText::FromName(NAME_None);
+	}
+	return FText::FromString(InDistortionHandlerPicker->HandlerDisplayName);
+}
+
+TSharedRef<SWidget> FDistortionHandlerPickerDetailCustomization::PopulateDistortionHandlerComboButton(FDistortionHandlerPicker* InDistortionHandlerPicker) const
+{
+	// Generate menu
+	FMenuBuilder MenuBuilder(true, nullptr);
+	MenuBuilder.BeginSection("DistortionHandlers", LOCTEXT("DistortionHandlers", "Distortion Sources For Target Camera"));
+	{
+		TArray<ULensDistortionModelHandlerBase*> Handlers;
+		UCameraCalibrationSubsystem* SubSystem = GEngine->GetEngineSubsystem<UCameraCalibrationSubsystem>();
+		if (SubSystem)
+		{
+			Handlers = SubSystem->GetDistortionModelHandlers(InDistortionHandlerPicker->TargetCameraComponent);
+		}
+
+		ULensDistortionModelHandlerBase* NullHandler = nullptr;
+		FString EmptyString;
+
+ 		//Always add a None entry
+		MenuBuilder.AddMenuEntry(
+			FText::FromName(NAME_None),
+			FText::FromName(NAME_None),
+			FSlateIcon(),
+			FUIAction(
+				FExecuteAction::CreateSP(this, &FDistortionHandlerPickerDetailCustomization::OnDistortionHandlerSelected, InDistortionHandlerPicker, NullHandler),
+				FCanExecuteAction(),
+				FIsActionChecked::CreateSP(this, &FDistortionHandlerPickerDetailCustomization::IsDistortionHandlerSelected, InDistortionHandlerPicker, EmptyString)
+			),
+			NAME_None,
+			EUserInterfaceActionType::RadioButton
+		);
+
+		for (ULensDistortionModelHandlerBase* Handler : Handlers)
+		{
+			MenuBuilder.AddMenuEntry(
+				FText::Format(LOCTEXT("DistortionHandlerLabel", "{0}"), FText::FromString(Handler->GetDisplayName())),
+				FText::Format(LOCTEXT("DistortionHandlerTooltip", "{0}"), FText::FromString(Handler->GetDisplayName())),
+				FSlateIcon(),
+				FUIAction(
+					FExecuteAction::CreateSP(this, &FDistortionHandlerPickerDetailCustomization::OnDistortionHandlerSelected, InDistortionHandlerPicker, Handler),
+					FCanExecuteAction(),
+					FIsActionChecked::CreateSP(this, &FDistortionHandlerPickerDetailCustomization::IsDistortionHandlerSelected, InDistortionHandlerPicker, Handler->GetDisplayName())
+				),
+				NAME_None,
+				EUserInterfaceActionType::RadioButton
+			);
+		}
+	}
+	MenuBuilder.EndSection();
+
+	return MenuBuilder.MakeWidget();
+}
+
+void FDistortionHandlerPickerDetailCustomization::OnDistortionHandlerSelected(FDistortionHandlerPicker* InDistortionHandlerPicker, ULensDistortionModelHandlerBase* InHandler) const
+{
+	const FScopedTransaction Transaction(LOCTEXT("SetDistortionHandler", "Set Distortion Source"));
+
+	TArray<UObject*> OuterObjects;
+	StructPropertyHandle->GetOuterObjects(OuterObjects);
+	for (UObject* Outer : OuterObjects)
+	{
+		Outer->Modify();
+	}
+
+	FGuid DistortionProducerID;
+	FString HandlerDisplayName;
+
+	if (InHandler)
+	{
+		DistortionProducerID = InHandler->GetDistortionProducerID();
+		HandlerDisplayName = InHandler->GetDisplayName();
+	}
+
+	InDistortionHandlerPicker->DistortionProducerID = DistortionProducerID;
+	InDistortionHandlerPicker->HandlerDisplayName = HandlerDisplayName;
+
+	StructPropertyHandle->NotifyPostChange();
+}
+
+bool FDistortionHandlerPickerDetailCustomization::IsDistortionHandlerSelected(FDistortionHandlerPicker* InDistortionHandlerPicker, FString InDisplayName) const
+{
+	return (InDistortionHandlerPicker->HandlerDisplayName.Equals(InDisplayName));
+}
+
+#undef LOCTEXT_NAMESPACE
\ No newline at end of file
diff --git a/Engine/Plugins/VirtualProduction/CameraCalibration/Source/CameraCalibrationEditor/Private/DistortionHandlerPickerDetailCustomization.h b/Engine/Plugins/VirtualProduction/CameraCalibration/Source/CameraCalibrationEditor/Private/DistortionHandlerPickerDetailCustomization.h
new file mode 100644
--- /dev/null
+++ b/Engine/Plugins/VirtualProduction/CameraCalibration/Source/CameraCalibrationEditor/Private/DistortionHandlerPickerDetailCustomization.h
@@ -0,0 +1,46 @@
+// Copyright Epic Games, Inc. All Rights Reserved.
+
+#pragma once
+
+#include "IPropertyTypeCustomization.h"
+
+#include "CameraCalibrationTypes.h"
+#include "CoreMinimal.h"
+#include "ScopedTransaction.h"
+#include "Widgets/Text/STextBlock.h"
+
+class UCineCameraComponent;
+class ULensDistortionModelHandlerBase;
+
+class FDistortionHandlerPickerDetailCustomization : public IPropertyTypeCustomization
+{
+public:
+	/** Makes a new instance of this detail layout class for a specific detail view requesting it. */
+	static TSharedRef<IPropertyTypeCustomization> MakeInstance();
+
+public:
+	//~ Begin IDetailCustomization interface
+	virtual void CustomizeHeader(TSharedRef<IPropertyHandle> PropertyHandle, FDetailWidgetRow& HeaderRow, IPropertyTypeCustomizationUtils& CustomizationUtils) override;
+	virtual void CustomizeChildren(TSharedRef<IPropertyHandle> PropertyHandle, IDetailChildrenBuilder& ChildBuilder, IPropertyTypeCustomizationUtils& CustomizationUtils) override {}
+	//~ End IDetailCustomization interface
+
+protected:
+	/** Construct the combo button widget to display a dropdown of the distortion sources that exist for the input picker */
+	TSharedRef<SWidget> BuildDistortionHandlerPickerWidget(FDistortionHandlerPicker* InDistortionHandlerPicker);
+
+	/** Construct a menu widget of distortion sources that exist for the input picker */
+	TSharedRef<SWidget> PopulateDistortionHandlerComboButton(FDistortionHandlerPicker* InDistortionHandlerPicker) const;
+
+	/** Delegate to run when a user makes a selection in the distortion source dropdown */
+	void OnDistortionHandlerSelected(FDistortionHandlerPicker* InDistortionHandlerPicker, ULensDistortionModelHandlerBase* InHandler) const;
+
+	/** Delegate to run to determine to determine whether the radio button should be filled in next to each dropdown option */
+	bool IsDistortionHandlerSelected(FDistortionHandlerPicker* InDistortionHandlerPicker, FString InDisplayName) const;
+
+	/** Delegate to run to get the text to display on the combo button */
+	FText OnGetButtonText(FDistortionHandlerPicker* InDistortionHandlerPicker) const;
+
+protected:
+	/** Cached property handle to the struct being customized */
+	TSharedPtr<IPropertyHandle> StructPropertyHandle;
+};
diff --git a/Engine/Plugins/VirtualProduction/LiveLinkCamera/Source/LiveLinkCamera/Private/LiveLinkCameraController.cpp b/Engine/Plugins/VirtualProduction/LiveLinkCamera/Source/LiveLinkCamera/Private/LiveLinkCameraController.cpp
--- a/Engine/Plugins/VirtualProduction/LiveLinkCamera/Source/LiveLinkCamera/Private/LiveLinkCameraController.cpp
+++ b/Engine/Plugins/VirtualProduction/LiveLinkCamera/Source/LiveLinkCamera/Private/LiveLinkCameraController.cpp
@@ -23,7 +23,6 @@
 #include "Kismet2/ComponentEditorUtils.h"
 #endif
 
-
 DEFINE_LOG_CATEGORY_STATIC(LogLiveLinkCameraController, Log, All);
 
 ULiveLinkCameraController::ULiveLinkCameraController()
@@ -32,6 +31,8 @@ ULiveLinkCameraController::ULiveLinkCameraController()
 	{
 		//Hook up to PostActorTick to handle nodal offset
 		FWorldDelegates::OnWorldPostActorTick.AddUObject(this, &ULiveLinkCameraController::OnPostActorTick);
+
+		DistortionProducerID = FGuid::NewGuid();
 	}
 }
 
@@ -157,6 +158,14 @@ void ULiveLinkCameraController::Cleanup()
 {
 	CleanupDistortion();
 
+	if (UCineCameraComponent* const CineCameraComponent = Cast<UCineCameraComponent>(AttachedComponent))
+	{
+		if (UCameraCalibrationSubsystem* SubSystem = GEngine->GetEngineSubsystem<UCameraCalibrationSubsystem>())
+		{
+			SubSystem->UnregisterDistortionModelHandler(CineCameraComponent, LensDistortionHandler);
+		}
+	}
+
 	FWorldDelegates::OnWorldPostActorTick.RemoveAll(this);
 }
 
@@ -166,6 +175,26 @@ void ULiveLinkCameraController::OnEvaluateRegistered()
 	bIsEncoderMappingNeeded = false;
 }
 
+void ULiveLinkCameraController::PostDuplicate(bool bDuplicateForPIE)
+{
+	Super::PostDuplicate(bDuplicateForPIE);
+
+	if (!DistortionProducerID.IsValid())
+	{
+		DistortionProducerID = FGuid::NewGuid();
+	}
+}
+
+void ULiveLinkCameraController::PostEditImport()
+{
+	Super::PostEditImport();
+
+	if (!DistortionProducerID.IsValid())
+	{
+		DistortionProducerID = FGuid::NewGuid();
+	}
+}
+
 #if WITH_EDITOR
 void ULiveLinkCameraController::PostEditChangeProperty(struct FPropertyChangedEvent& PropertyChangedEvent)
 {
@@ -191,10 +220,16 @@ void ULiveLinkCameraController::PostEditChangeProperty(struct FPropertyChangedEv
 			{
 				// Static lens model data may not be available because the controller may not be ticking, so query the subsystem for any distortion handler
 				UCameraCalibrationSubsystem* SubSystem = GEngine->GetEngineSubsystem<UCameraCalibrationSubsystem>();
-				ULensDistortionModelHandlerBase* Handler = SubSystem->GetDistortionModelHandler(CineCameraComponent);
+				FDistortionHandlerPicker DistortionHandlerPicker;
+				DistortionHandlerPicker.TargetCameraComponent = CineCameraComponent;
+				DistortionHandlerPicker.DistortionProducerID = DistortionProducerID;
+				ULensDistortionModelHandlerBase* Handler = SubSystem->FindDistortionModelHandler(DistortionHandlerPicker);
 
 				// Try to remove the post-process material from the cine camera's blendables. 
-				CineCameraComponent->RemoveBlendable(Handler->GetDistortionMID());
+				if (Handler)
+				{
+					CineCameraComponent->RemoveBlendable(Handler->GetDistortionMID());
+				}
 
 				const float OriginalFOV = FMath::RadiansToDegrees(2.0f * FMath::Atan(CineCameraComponent->Filmback.SensorWidth / (2.0f * UndistortedFocalLength)));
 				CineCameraComponent->SetFieldOfView(OriginalFOV);
@@ -399,8 +434,11 @@ void ULiveLinkCameraController::ApplyDistortion(ULensFile* LensFile, UCineCamera
 
 	if (LensFile)
 	{
-		UCameraCalibrationSubsystem* SubSystem = GEngine->GetEngineSubsystem<UCameraCalibrationSubsystem>();
-		LensDistortionHandler = SubSystem->FindOrCreateDistortionModelHandler(CineCameraComponent, LensFile->LensInfo.LensModel);
+		UCameraCalibrationSubsystem* const SubSystem = GEngine->GetEngineSubsystem<UCameraCalibrationSubsystem>();
+		const FString HandlerDisplayName = FString::Format(TEXT("{0} (Lens File)"), { LensFile->GetFName().ToString() });
+
+		FDistortionHandlerPicker DistortionHandlerPicker = { CineCameraComponent, DistortionProducerID, HandlerDisplayName };
+		LensDistortionHandler = SubSystem->FindOrCreateDistortionModelHandler(DistortionHandlerPicker, LensFile->LensInfo.LensModel);
 	}
 
 	if (LensDistortionHandler)
diff --git a/Engine/Plugins/VirtualProduction/LiveLinkCamera/Source/LiveLinkCamera/Public/LiveLinkCameraController.h b/Engine/Plugins/VirtualProduction/LiveLinkCamera/Source/LiveLinkCamera/Public/LiveLinkCameraController.h
--- a/Engine/Plugins/VirtualProduction/LiveLinkCamera/Source/LiveLinkCamera/Public/LiveLinkCameraController.h
+++ b/Engine/Plugins/VirtualProduction/LiveLinkCamera/Source/LiveLinkCamera/Public/LiveLinkCameraController.h
@@ -105,6 +105,10 @@ protected:
 	UPROPERTY()
 	FVector OriginalCameraLocation;
 
+	/** Unique identifier representing the source of distortion data */
+	UPROPERTY(DuplicateTransient)
+	FGuid DistortionProducerID;
+
 	/** Used to control which data from LiveLink is actually applied to camera */
 	UPROPERTY(EditAnywhere, Category="Settings")
 	FLiveLinkCameraControllerUpdateFlags UpdateFlags;
@@ -147,6 +151,8 @@ public:
 	//~ End ULiveLinkControllerBase interface
 
 	//~ Begin UObject interface
+	virtual void PostDuplicate(bool bDuplicateForPIE) override;
+	virtual void PostEditImport() override;
 	virtual void PostLoad() override;
 #if WITH_EDITOR
 	virtual void PostEditChangeProperty(struct FPropertyChangedEvent& PropertyChangedEvent) override;
diff --git a/Engine/Plugins/VirtualProduction/LiveLinkLens/Source/LiveLinkLens/Private/LiveLinkLensController.cpp b/Engine/Plugins/VirtualProduction/LiveLinkLens/Source/LiveLinkLens/Private/LiveLinkLensController.cpp
--- a/Engine/Plugins/VirtualProduction/LiveLinkLens/Source/LiveLinkLens/Private/LiveLinkLensController.cpp
+++ b/Engine/Plugins/VirtualProduction/LiveLinkLens/Source/LiveLinkLens/Private/LiveLinkLensController.cpp
@@ -3,12 +3,23 @@
 #include "LiveLinkLensController.h"
 
 #include "CameraCalibrationSubsystem.h"
+#include "Engine/Engine.h"
+#include "Engine/World.h"
+#include "LiveLinkComponentController.h"
 #include "LiveLinkLensRole.h"
 #include "LiveLinkLensTypes.h"
 #include "Logging/LogMacros.h"
 
+
 DEFINE_LOG_CATEGORY_STATIC(LogLiveLinkLensController, Log, All);
 
+ULiveLinkLensController::ULiveLinkLensController() 
+{
+	if (!HasAnyFlags(RF_ArchetypeObject | RF_ClassDefaultObject))
+	{
+		DistortionProducerID = FGuid::NewGuid();
+	}
+}
 
 void ULiveLinkLensController::Tick(float DeltaTime, const FLiveLinkSubjectFrameData& SubjectData)
 {
@@ -19,10 +30,16 @@ void ULiveLinkLensController::Tick(float DeltaTime, const FLiveLinkSubjectFrameD
 	{
 		if (UCineCameraComponent* const CineCameraComponent = Cast<UCineCameraComponent>(AttachedComponent))
 		{
-			UCameraCalibrationSubsystem* SubSystem = GEngine->GetEngineSubsystem<UCameraCalibrationSubsystem>();
-			TSubclassOf<ULensModel> LensModel = SubSystem->GetRegisteredLensModel(StaticData->LensModel);
-			LensDistortionHandler = SubSystem->FindOrCreateDistortionModelHandler(CineCameraComponent, LensModel, EHandlerOverrideMode::SoftOverride);
+			UCameraCalibrationSubsystem* const SubSystem = GEngine->GetEngineSubsystem<UCameraCalibrationSubsystem>();
+			const TSubclassOf<ULensModel> LensModel = SubSystem->GetRegisteredLensModel(StaticData->LensModel);
+
+			const FName SubjectName = SelectedSubject.Subject.Name;
+			const FText RoleName = SelectedSubject.Role->GetDefaultObject<ULiveLinkRole>()->GetDisplayName();
+			const FString HandlerDisplayName = FString::Format(TEXT("'{0}' LiveLink Subject ({1} Role)"), { SubjectName.ToString(), RoleName.ToString() });
 
+			FDistortionHandlerPicker DistortionHandlerPicker = { CineCameraComponent, DistortionProducerID, HandlerDisplayName };
+			LensDistortionHandler = SubSystem->FindOrCreateDistortionModelHandler(DistortionHandlerPicker, LensModel);
+			
 			//To keep track of an updated MID from the handler
 			UMaterialInstanceDynamic* NewDistortionMID = nullptr;
 
@@ -134,6 +151,37 @@ void ULiveLinkLensController::SetAttachedComponent(UActorComponent* ActorCompone
 void ULiveLinkLensController::Cleanup()
 {
 	CleanupDistortion();
+
+	if (UCineCameraComponent* const CineCameraComponent = Cast<UCineCameraComponent>(AttachedComponent))
+	{
+		if (UCameraCalibrationSubsystem* SubSystem = GEngine->GetEngineSubsystem<UCameraCalibrationSubsystem>())
+		{
+			SubSystem->UnregisterDistortionModelHandler(CineCameraComponent, LensDistortionHandler);
+		}
+	}
+}
+
+void ULiveLinkLensController::PostDuplicate(bool bDuplicateForPIE)
+{
+	Super::PostDuplicate(bDuplicateForPIE);
+
+	// When this controller is duplicated (e.g. for PIE), the duplicated controller needs a new unique ID
+	if (!DistortionProducerID.IsValid())
+	{
+		DistortionProducerID = FGuid::NewGuid();
+	}
+}
+
+void ULiveLinkLensController::PostEditImport()
+{
+	Super::PostEditImport();
+
+	// PostDuplicate is not called on components during actor duplication, such as alt-drag and copy-paste, so PostEditImport covers those duplication cases
+	// When this controller is duplicated in those cases, the duplicated controller needs a new unique ID
+	if (!DistortionProducerID.IsValid())
+	{
+		DistortionProducerID = FGuid::NewGuid();
+	}
 }
 
 #if WITH_EDITOR
@@ -161,10 +209,16 @@ void ULiveLinkLensController::PostEditChangeProperty(struct FPropertyChangedEven
 			{
 				// Static lens model data may not be available because the controller may not be ticking, so query the subsystem for any distortion handler
 				UCameraCalibrationSubsystem* SubSystem = GEngine->GetEngineSubsystem<UCameraCalibrationSubsystem>();
-				ULensDistortionModelHandlerBase* Handler = SubSystem->GetDistortionModelHandler(CineCameraComponent);
+				FDistortionHandlerPicker DistortionHandlerPicker;
+				DistortionHandlerPicker.TargetCameraComponent = CineCameraComponent;
+				DistortionHandlerPicker.DistortionProducerID = DistortionProducerID;
+				ULensDistortionModelHandlerBase* Handler = SubSystem->FindDistortionModelHandler(DistortionHandlerPicker);
 
 				// Try to remove the post-process material from the cine camera's blendables. 
-				CineCameraComponent->RemoveBlendable(Handler->GetDistortionMID());
+				if (Handler)
+				{
+					CineCameraComponent->RemoveBlendable(Handler->GetDistortionMID());
+				}
 
 				// Restore the original FOV of the 
 				const float OriginalFOV = FMath::RadiansToDegrees(2.0f * FMath::Atan(CineCameraComponent->Filmback.SensorWidth / (2.0f * UndistortedFocalLength)));
diff --git a/Engine/Plugins/VirtualProduction/LiveLinkLens/Source/LiveLinkLens/Public/LiveLinkLensController.h b/Engine/Plugins/VirtualProduction/LiveLinkLens/Source/LiveLinkLens/Public/LiveLinkLensController.h
--- a/Engine/Plugins/VirtualProduction/LiveLinkLens/Source/LiveLinkLens/Public/LiveLinkLensController.h
+++ b/Engine/Plugins/VirtualProduction/LiveLinkLens/Source/LiveLinkLens/Public/LiveLinkLensController.h
@@ -19,7 +19,7 @@ class LIVELINKLENS_API ULiveLinkLensController : public ULiveLinkControllerBase
 	GENERATED_BODY()
 
 public:
-	ULiveLinkLensController() {}
+	ULiveLinkLensController();
 
 	//~ Begin ULiveLinkControllerBase interface
 	virtual void Tick(float DeltaTime, const FLiveLinkSubjectFrameData& SubjectData) override;
@@ -30,6 +30,8 @@ public:
 	//~ End ULiveLinkControllerBase interface
 
 	//~ Begin UObject Interface
+	virtual void PostDuplicate(bool bDuplicateForPIE) override;
+	virtual void PostEditImport() override;
 #if WITH_EDITOR
 	virtual void PostEditChangeProperty(struct FPropertyChangedEvent& PropertyChangedEvent) override;
 #endif
@@ -48,11 +50,11 @@ protected:
 
 protected:
 	/** Whether or not to apply a post-process distortion effect directly to the attached CineCamera */
-	UPROPERTY(EditAnywhere, Category = "Lens Distortion")
+	UPROPERTY(EditAnywhere, Category = "Camera Calibration")
 	bool bApplyDistortion = false;
 
 	/** Cached distortion handler associated with attached camera component */
-	UPROPERTY(EditAnywhere, Category = "Lens Distortion", Transient)
+	UPROPERTY(EditAnywhere, Category = "Camera Calibration", Transient)
 	ULensDistortionModelHandlerBase* LensDistortionHandler = nullptr;
 
 	/** Cached distortion MID the handler produced. Used to clean up old one in case it changes */
@@ -63,9 +65,13 @@ protected:
  	UPROPERTY()
 	float UndistortedFocalLength = 50.0f;
 
+	/** Unique identifier representing the source of distortion data */
+	UPROPERTY(DuplicateTransient)
+	FGuid DistortionProducerID;
+
 	/** Keep track of what needs to be setup to apply distortion */
 	bool bIsDistortionSetup = false;
 
-	//Last value used to detect changes made by the user and update our original caches
+	/** Last value used to detect changes made by the user and update our original caches */
 	float LastFocalLength = -1.0f;
 };
\ No newline at end of file
