- Fix CIS (missing function parameter)

diff --git a/Engine/Plugins/VirtualProduction/CameraCalibration/Source/CameraCalibrationEditor/Private/DistortionHandlerPickerDetailCustomization.cpp b/Engine/Plugins/VirtualProduction/CameraCalibration/Source/CameraCalibrationEditor/Private/DistortionHandlerPickerDetailCustomization.cpp
--- a/Engine/Plugins/VirtualProduction/CameraCalibration/Source/CameraCalibrationEditor/Private/DistortionHandlerPickerDetailCustomization.cpp
+++ b/Engine/Plugins/VirtualProduction/CameraCalibration/Source/CameraCalibrationEditor/Private/DistortionHandlerPickerDetailCustomization.cpp
@@ -150,7 +150,7 @@ void FDistortionHandlerPickerDetailCustomization::OnDistortionHandlerSelected(FD
 	InDistortionHandlerPicker->DistortionProducerID = DistortionProducerID;
 	InDistortionHandlerPicker->HandlerDisplayName = HandlerDisplayName;
 
-	StructPropertyHandle->NotifyPostChange();
+	StructPropertyHandle->NotifyPostChange(EPropertyChangeType::Unspecified);
 }
 
 bool FDistortionHandlerPickerDetailCustomization::IsDistortionHandlerSelected(FDistortionHandlerPicker* InDistortionHandlerPicker, FString InDisplayName) const
