- Create a new lens distortion component to apply a distortion post-process effect on a camera, and optionally also drive the distortion state based on current camera settings and a lens file
- Remove "None" entry from distortion handler picker detail customization when there are valid sources in the list

diff --git a/Engine/Plugins/VirtualProduction/CameraCalibration/Source/CameraCalibration/CameraCalibration.Build.cs b/Engine/Plugins/VirtualProduction/CameraCalibration/Source/CameraCalibration/CameraCalibration.Build.cs
--- a/Engine/Plugins/VirtualProduction/CameraCalibration/Source/CameraCalibration/CameraCalibration.Build.cs
+++ b/Engine/Plugins/VirtualProduction/CameraCalibration/Source/CameraCalibration/CameraCalibration.Build.cs
@@ -22,6 +22,7 @@ namespace UnrealBuildTool.Rules
 			PrivateDependencyModuleNames.AddRange(
 				new string[]
                 {
+					"LiveLinkComponents",
                     "Projects",
 				}
 			);
diff --git a/Engine/Plugins/VirtualProduction/CameraCalibration/Source/CameraCalibration/Private/LensDistortionComponent.cpp b/Engine/Plugins/VirtualProduction/CameraCalibration/Source/CameraCalibration/Private/LensDistortionComponent.cpp
new file mode 100644
--- /dev/null
+++ b/Engine/Plugins/VirtualProduction/CameraCalibration/Source/CameraCalibration/Private/LensDistortionComponent.cpp
@@ -0,0 +1,311 @@
+// Copyright Epic Games, Inc. All Rights Reserved.
+
+#include "LensDistortionComponent.h"
+
+#include "CameraCalibrationSubsystem.h"
+#include "CineCameraComponent.h"
+#include "LiveLinkComponentController.h"
+
+ULensDistortionComponent::ULensDistortionComponent()
+{
+	PrimaryComponentTick.bCanEverTick = true;
+	PrimaryComponentTick.bStartWithTickEnabled = true;
+	PrimaryComponentTick.TickGroup = ETickingGroup::TG_LastDemotable;
+	bTickInEditor = true;
+
+	if (!HasAnyFlags(RF_ArchetypeObject | RF_ClassDefaultObject))
+	{
+		DistortionProducerID = FGuid::NewGuid();
+	}
+}
+
+void ULensDistortionComponent::OnRegister()
+{
+	Super::OnRegister();
+
+	InitDefaultCamera();
+}
+
+void ULensDistortionComponent::DestroyComponent(bool bPromoteChildren /*= false*/)
+{
+	Super::DestroyComponent(bPromoteChildren);
+
+	if (UCineCameraComponent* const CineCameraComponent = Cast<UCineCameraComponent>(TargetCameraComponent.GetComponent(GetOwner())))
+	{
+		CleanupDistortion(CineCameraComponent);
+
+		if (UCameraCalibrationSubsystem* SubSystem = GEngine->GetEngineSubsystem<UCameraCalibrationSubsystem>())
+		{
+			SubSystem->UnregisterDistortionModelHandler(CineCameraComponent, ProducedLensDistortionHandler);
+		}
+	}
+}
+
+void ULensDistortionComponent::TickComponent(float DeltaTime, ELevelTick TickType, FActorComponentTickFunction* ThisTickFunction)
+{
+	if (UCineCameraComponent* const CineCameraComponent = Cast<UCineCameraComponent>(TargetCameraComponent.GetComponent(GetOwner())))
+	{
+		// Add a tick prerequisite on any live link components also on this camera
+		// LiveLink might set the focal length of the target camera, so it is important that this component ticks after that
+		TInlineComponentArray<ULiveLinkComponentController*> LiveLinkComponents;
+		CineCameraComponent->GetOwner()->GetComponents(LiveLinkComponents);
+
+		for (UActorComponent* LiveLinkComponent : LiveLinkComponents)
+		{
+			AddTickPrerequisiteComponent(LiveLinkComponent);
+		}
+
+		DistortionSource.TargetCameraComponent = CineCameraComponent;
+
+		// Evaluate the selected lens file for distortion data using the current settings of the target camera
+		ULensFile* LensFile = LensFilePicker.GetLensFile();
+		if (bEvaluateLensFileForDistortion && LensFile)
+		{
+			if (!ProducedLensDistortionHandler)
+			{
+				CreateDistortionHandlerForLensFile();
+			}
+
+			//Go through the lens file to get distortion data based on FIZ
+			//Our handler's displacement map will get updated
+			FDistortionData DistortionData;
+			const FVector2D CurrentSensorDimensions = FVector2D(CineCameraComponent->Filmback.SensorWidth, CineCameraComponent->Filmback.SensorHeight);
+
+			LensFile->EvaluateDistortionData(CineCameraComponent->CurrentFocusDistance, UndistortedFocalLength, CurrentSensorDimensions, ProducedLensDistortionHandler, DistortionData);
+		}
+	
+		if (bApplyDistortion)
+		{
+			// Get the lens distortion handler for the target camera and distortion source
+			ULensDistortionModelHandlerBase* LensDistortionHandler = nullptr;
+			UCameraCalibrationSubsystem* SubSystem = GEngine->GetEngineSubsystem<UCameraCalibrationSubsystem>();
+			if (SubSystem)
+			{
+				LensDistortionHandler = SubSystem->FindDistortionModelHandler(DistortionSource);
+			}
+
+			if (LensDistortionHandler)
+			{
+				// Get the current distortion MID from the lens distortion handler
+				UMaterialInstanceDynamic* NewDistortionMID = LensDistortionHandler->GetDistortionMID();
+
+				// If the MID has changed
+				if (LastDistortionMID != NewDistortionMID)
+				{
+					CineCameraComponent->RemoveBlendable(LastDistortionMID);
+					CineCameraComponent->AddOrUpdateBlendable(NewDistortionMID);
+				}
+
+				// Cache the latest distortion MID
+				LastDistortionMID = NewDistortionMID;
+
+				// Check for focal length modifications made from outside of this component to update the "undistorted" focal length used
+				if (LastFocalLength != CineCameraComponent->CurrentFocalLength)
+				{
+					UndistortedFocalLength = CineCameraComponent->CurrentFocalLength;
+				}
+
+				// Get the overscan factor and use it to modify the target camera's FOV
+				const float OverscanFactor = LensDistortionHandler->GetOverscanFactor();
+				const float OverscanSensorWidth = CineCameraComponent->Filmback.SensorWidth * OverscanFactor;
+				const float OverscanFOV = FMath::RadiansToDegrees(2.0f * FMath::Atan(OverscanSensorWidth / (2.0f * UndistortedFocalLength)));
+				CineCameraComponent->SetFieldOfView(OverscanFOV);
+
+				//Stamp last applied focal length to detect if user has changed it
+				LastFocalLength = CineCameraComponent->CurrentFocalLength;
+
+				bIsDistortionSetup = true;
+			}
+			else
+			{
+				CleanupDistortion(CineCameraComponent);
+			}
+		}
+	}
+
+	Super::TickComponent(DeltaTime, TickType, ThisTickFunction);
+}
+
+void ULensDistortionComponent::PostEditChangeProperty(struct FPropertyChangedEvent& PropertyChangedEvent)
+{
+ 	const FName PropertyName = (PropertyChangedEvent.Property != NULL) ? PropertyChangedEvent.Property->GetFName() : NAME_None;
+ 
+ 	if (PropertyName == GET_MEMBER_NAME_CHECKED(ULensDistortionComponent, bApplyDistortion))
+ 	{
+ 		if (UCineCameraComponent* const CineCameraComponent = Cast<UCineCameraComponent>(TargetCameraComponent.GetComponent(GetOwner())))
+ 		{
+ 			if (!bApplyDistortion)
+ 			{
+				CleanupDistortion(CineCameraComponent);
+ 			}
+ 		}
+ 	}
+	else if (PropertyName == GET_MEMBER_NAME_CHECKED(ULensDistortionComponent, bEvaluateLensFileForDistortion))
+	{
+		if (bEvaluateLensFileForDistortion)
+		{
+			// Create a new distortion handler if one does not yet exist for the selected lens file
+			if (!ProducedLensDistortionHandler)
+			{
+				CreateDistortionHandlerForLensFile();
+			}
+		}
+		else
+		{
+			// Reset the distortion source settings
+			if (DistortionSource.DistortionProducerID == DistortionProducerID)
+			{
+				DistortionSource.HandlerDisplayName.Empty();
+				DistortionSource.DistortionProducerID.Invalidate();
+			}
+
+			// Unregister the produced handler from the subsystem
+			if (UCameraCalibrationSubsystem* SubSystem = GEngine->GetEngineSubsystem<UCameraCalibrationSubsystem>())
+			{
+				if (UCineCameraComponent* const CineCameraComponent = Cast<UCineCameraComponent>(TargetCameraComponent.GetComponent(GetOwner())))
+				{
+					SubSystem->UnregisterDistortionModelHandler(CineCameraComponent, ProducedLensDistortionHandler);
+				}
+			}
+
+			ProducedLensDistortionHandler = nullptr;
+		}
+	}
+	else if (PropertyName == GET_MEMBER_NAME_CHECKED(ULensDistortionComponent, TargetCameraComponent))
+	{
+		// Clean up distortion on the last target camera
+		if (LastCameraComponent)
+		{
+			CleanupDistortion(LastCameraComponent);
+
+			// Unregister the produced handler (if one exists) from the subsystem for the last camera component
+			if (ProducedLensDistortionHandler)
+			{
+				if (UCameraCalibrationSubsystem* SubSystem = GEngine->GetEngineSubsystem<UCameraCalibrationSubsystem>())
+				{
+					SubSystem->UnregisterDistortionModelHandler(LastCameraComponent, ProducedLensDistortionHandler);
+				}
+
+				ProducedLensDistortionHandler = nullptr;
+			}
+		}
+
+		LastCameraComponent = Cast<UCineCameraComponent>(TargetCameraComponent.GetComponent(GetOwner()));
+		DistortionSource.TargetCameraComponent = LastCameraComponent;
+
+		// Create a new handler for the new target camera component
+		if (bEvaluateLensFileForDistortion)
+		{
+			CreateDistortionHandlerForLensFile();
+		}
+	}
+	else if (PropertyName == GET_MEMBER_NAME_CHECKED(FLensFilePicker, LensFile))
+	{
+		if (bEvaluateLensFileForDistortion)
+		{
+			CreateDistortionHandlerForLensFile();
+		}
+	}
+
+	Super::PostEditChangeProperty(PropertyChangedEvent);
+}
+
+void ULensDistortionComponent::PostDuplicate(bool bDuplicateForPIE)
+{
+	Super::PostDuplicate(bDuplicateForPIE);
+
+	InitDefaultCamera();
+
+	// Duplicate the last focal length (e.g. for PIE)
+	if (bIsDistortionSetup)
+	{
+		if (UCineCameraComponent* const CineCameraComponent = Cast<UCineCameraComponent>(TargetCameraComponent.GetComponent(GetOwner())))
+		{
+			LastFocalLength = CineCameraComponent->CurrentFocalLength;
+		}
+	}
+
+	// When this component is duplicated (e.g. for PIE), the duplicated component needs a new unique ID
+	if (!DistortionProducerID.IsValid())
+	{
+		DistortionProducerID = FGuid::NewGuid();
+	}
+}
+
+void ULensDistortionComponent::PostEditImport()
+{
+	Super::PostEditImport();
+
+	InitDefaultCamera();
+
+	// PostDuplicate is not called on components during actor duplication, such as alt-drag and copy-paste, so PostEditImport covers those duplication cases
+	// Duplicate the last focal length in these cases
+	if (bIsDistortionSetup)
+	{
+		if (UCineCameraComponent* const CineCameraComponent = Cast<UCineCameraComponent>(TargetCameraComponent.GetComponent(GetOwner())))
+		{
+			LastFocalLength = CineCameraComponent->CurrentFocalLength;
+		}
+	}
+
+	// When this component is duplicated (e.g. copy-paste), the duplicated component needs a new unique ID
+	if (!DistortionProducerID.IsValid())
+	{
+		DistortionProducerID = FGuid::NewGuid();
+	}
+}
+
+void ULensDistortionComponent::CleanupDistortion(UCineCameraComponent* const CineCameraComponent)
+{
+	if (bIsDistortionSetup)
+	{
+		// Remove the last distortion MID that was applied to the target camera component
+		if (LastDistortionMID)
+		{
+			CineCameraComponent->RemoveBlendable(LastDistortionMID);
+			LastDistortionMID = nullptr;
+		}
+
+		// Restore the original FOV of the target camera
+		const float UndistortedFOV = FMath::RadiansToDegrees(2.0f * FMath::Atan(CineCameraComponent->Filmback.SensorWidth / (2.0f * UndistortedFocalLength)));
+		CineCameraComponent->SetFieldOfView(UndistortedFOV);
+
+		LastFocalLength = CineCameraComponent->CurrentFocalLength;
+	}
+
+	bIsDistortionSetup = false;
+}
+
+void ULensDistortionComponent::InitDefaultCamera()
+{
+	UCineCameraComponent* const CineCameraComponent = Cast<UCineCameraComponent>(TargetCameraComponent.GetComponent(GetOwner()));
+	if (!CineCameraComponent)
+	{
+		// Find all CineCameraComponents on the same actor as this component and set the first one to be the target
+		TInlineComponentArray<UCineCameraComponent*> CineCameraComponents;
+		GetOwner()->GetComponents(CineCameraComponents);
+		if (CineCameraComponents.Num() > 0)
+		{
+			TargetCameraComponent.ComponentProperty = CineCameraComponents[0]->GetFName();
+			LastCameraComponent = CineCameraComponents[0];
+		}
+	}
+}
+
+void ULensDistortionComponent::CreateDistortionHandlerForLensFile()
+{
+	ULensFile* const LensFile = LensFilePicker.GetLensFile();
+
+	// Set the distortion source settings to match this component and the selected lens file
+	if (LensFile)
+	{
+		DistortionSource.HandlerDisplayName = FString::Format(TEXT("{0} (Lens File) (Lens Distortion Component)"), { LensFile->GetFName().ToString() });
+		DistortionSource.DistortionProducerID = DistortionProducerID;
+
+		UCameraCalibrationSubsystem* const SubSystem = GEngine->GetEngineSubsystem<UCameraCalibrationSubsystem>();
+		if (SubSystem)
+		{
+			ProducedLensDistortionHandler = SubSystem->FindOrCreateDistortionModelHandler(DistortionSource, LensFile->LensInfo.LensModel);
+		}
+	}
+}
\ No newline at end of file
diff --git a/Engine/Plugins/VirtualProduction/CameraCalibration/Source/CameraCalibration/Public/LensDistortionComponent.h b/Engine/Plugins/VirtualProduction/CameraCalibration/Source/CameraCalibration/Public/LensDistortionComponent.h
new file mode 100644
--- /dev/null
+++ b/Engine/Plugins/VirtualProduction/CameraCalibration/Source/CameraCalibration/Public/LensDistortionComponent.h
@@ -0,0 +1,98 @@
+// Copyright Epic Games, Inc. All Rights Reserved.
+
+#pragma once
+
+#include "Components/ActorComponent.h"
+
+#include "CameraCalibrationTypes.h"
+#include "LensFile.h"
+
+#include "LensDistortionComponent.generated.h"
+
+class UCineCameraComponent;
+
+/** Component for applying a post-process lens distortion effect to a CineCameraComponent on the same actor */
+UCLASS(HideCategories = (Tags, Activation, Cooking, AssetUserData, Collision), meta = (BlueprintSpawnableComponent))
+class CAMERACALIBRATION_API ULensDistortionComponent : public UActorComponent
+{
+	GENERATED_BODY()
+
+public:
+	ULensDistortionComponent();
+
+	//~ Begin UActorComponent interface
+	virtual void OnRegister() override;
+	virtual void TickComponent(float DeltaTime, ELevelTick TickType, FActorComponentTickFunction* ThisTickFunction) override;
+	virtual void DestroyComponent(bool bPromoteChildren = false) override;
+	//~ End UActorComponent interface
+
+	//~ UObject interface
+#if WITH_EDITOR
+	virtual void PostEditChangeProperty(struct FPropertyChangedEvent& PropertyChangedEvent) override;
+#endif //WITH_EDITOR
+	virtual void PostDuplicate(bool bDuplicateForPIE) override;
+	virtual void PostEditImport() override;
+	//~ End UObject interface
+
+private:
+	/** If TargetCameraComponent is not set, initialize it to the first CineCameraComponent on the same actor as this component */
+	void InitDefaultCamera();
+
+	/** Remove the last distortion MID applied to the input CineCameraComponent and reset its FOV to use no overscan */
+	void CleanupDistortion(UCineCameraComponent* CineCameraComponent);
+
+	/** Register a new lens distortion handler with the camera calibration subsystem using the selected lens file */
+	void CreateDistortionHandlerForLensFile();
+
+protected:
+	/** The CineCameraComponent on which to apply the post-process distortion effect */
+	UPROPERTY(EditInstanceOnly, Category = "Default", meta = (UseComponentPicker, AllowedClasses = "CineCameraComponent"))
+	FComponentReference TargetCameraComponent;
+
+	/** Structure used to query the camera calibration subsystem for a lens distortion model handler */
+	UPROPERTY(EditAnywhere, BlueprintReadWrite, Category = "Default")
+	FDistortionHandlerPicker DistortionSource;
+
+	/** Whether or not to apply distortion to the target camera component */
+	UPROPERTY(EditAnywhere, BlueprintReadWrite, Category = "Default")
+	bool bApplyDistortion = false;
+
+	/** Whether a distortion effect is currently being applied to the target camera component */
+	UPROPERTY()
+	bool bIsDistortionSetup = false;
+
+	/** Cached focal length of the target camera with no overscan applied */
+	UPROPERTY()
+	float UndistortedFocalLength = 35.0f;
+
+	/** Cached MID last applied to the target camera */
+	UPROPERTY()
+	UMaterialInstanceDynamic* LastDistortionMID = nullptr;
+
+	/** Cached most recent target camera, used to clean up the old camera when the user changes the target */
+	UPROPERTY(Transient)
+	UCineCameraComponent* LastCameraComponent = nullptr;
+
+	/** 
+	 * Whether to use the specified lens file to drive distortion
+	 * Enabling this will create a new distortion source and automatically set this component to use it
+	 */
+	UPROPERTY(EditAnywhere, AdvancedDisplay, Category = "Default")
+	bool bEvaluateLensFileForDistortion;
+	
+	/** Lens File used to drive distortion with current camera settings */
+	UPROPERTY(EditAnywhere, AdvancedDisplay, Category = "Default", meta = (EditCondition="bEvaluateLensFileForDistortion"))
+	FLensFilePicker LensFilePicker;
+
+	/** Distortion handler produced by this component */
+	UPROPERTY(Transient)
+	ULensDistortionModelHandlerBase* ProducedLensDistortionHandler = nullptr;
+
+	/** Unique identifier representing the source of distortion data */
+	UPROPERTY(DuplicateTransient)
+	FGuid DistortionProducerID;
+
+private:
+	/** Cached focal length of the target camera, used to detect changes to this value from outside of this component */
+	float LastFocalLength = -1.0f;
+};
\ No newline at end of file
diff --git a/Engine/Plugins/VirtualProduction/CameraCalibration/Source/CameraCalibrationEditor/Private/DistortionHandlerPickerDetailCustomization.cpp b/Engine/Plugins/VirtualProduction/CameraCalibration/Source/CameraCalibrationEditor/Private/DistortionHandlerPickerDetailCustomization.cpp
--- a/Engine/Plugins/VirtualProduction/CameraCalibration/Source/CameraCalibrationEditor/Private/DistortionHandlerPickerDetailCustomization.cpp
+++ b/Engine/Plugins/VirtualProduction/CameraCalibration/Source/CameraCalibrationEditor/Private/DistortionHandlerPickerDetailCustomization.cpp
@@ -89,22 +89,25 @@ TSharedRef<SWidget> FDistortionHandlerPickerDetailCustomization::PopulateDistort
 			Handlers = SubSystem->GetDistortionModelHandlers(InDistortionHandlerPicker->TargetCameraComponent);
 		}
 
-		ULensDistortionModelHandlerBase* NullHandler = nullptr;
-		FString EmptyString;
-
- 		//Always add a None entry
-		MenuBuilder.AddMenuEntry(
-			FText::FromName(NAME_None),
-			FText::FromName(NAME_None),
-			FSlateIcon(),
-			FUIAction(
-				FExecuteAction::CreateSP(this, &FDistortionHandlerPickerDetailCustomization::OnDistortionHandlerSelected, InDistortionHandlerPicker, NullHandler),
-				FCanExecuteAction(),
-				FIsActionChecked::CreateSP(this, &FDistortionHandlerPickerDetailCustomization::IsDistortionHandlerSelected, InDistortionHandlerPicker, EmptyString)
-			),
-			NAME_None,
-			EUserInterfaceActionType::RadioButton
-		);
+		if (Handlers.Num() == 0)
+		{
+			ULensDistortionModelHandlerBase* NullHandler = nullptr;
+			FString EmptyString;
+
+			//Always add a None entry
+			MenuBuilder.AddMenuEntry(
+				FText::FromName(NAME_None),
+				FText::FromName(NAME_None),
+				FSlateIcon(),
+				FUIAction(
+					FExecuteAction::CreateSP(this, &FDistortionHandlerPickerDetailCustomization::OnDistortionHandlerSelected, InDistortionHandlerPicker, NullHandler),
+					FCanExecuteAction(),
+					FIsActionChecked::CreateSP(this, &FDistortionHandlerPickerDetailCustomization::IsDistortionHandlerSelected, InDistortionHandlerPicker, EmptyString)
+				),
+				NAME_None,
+				EUserInterfaceActionType::RadioButton
+			);
+		}
 
 		for (ULensDistortionModelHandlerBase* Handler : Handlers)
 		{
diff --git a/Engine/Plugins/VirtualProduction/LiveLinkCamera/Source/LiveLinkCamera/Private/LiveLinkCameraController.cpp b/Engine/Plugins/VirtualProduction/LiveLinkCamera/Source/LiveLinkCamera/Private/LiveLinkCameraController.cpp
--- a/Engine/Plugins/VirtualProduction/LiveLinkCamera/Source/LiveLinkCamera/Private/LiveLinkCameraController.cpp
+++ b/Engine/Plugins/VirtualProduction/LiveLinkCamera/Source/LiveLinkCamera/Private/LiveLinkCameraController.cpp
@@ -76,7 +76,7 @@ void ULiveLinkCameraController::Tick(float DeltaTime, const FLiveLinkSubjectFram
 				LastFilmback = CineCameraComponent->Filmback;
 
 				ApplyFIZ(SelectedLensFile, CineCameraComponent, StaticData, FrameData);
-				ApplyDistortion(SelectedLensFile, CineCameraComponent);
+				ApplyDistortion(SelectedLensFile, CineCameraComponent, StaticData, FrameData);
 			}
 
 #if WITH_EDITORONLY_DATA
@@ -106,43 +106,13 @@ TSubclassOf<UActorComponent> ULiveLinkCameraController::GetDesiredComponentClass
 
 void ULiveLinkCameraController::SetAttachedComponent(UActorComponent* ActorComponent)
 {
-	const bool bHasChangedComponent = (ActorComponent != AttachedComponent);
-	if (bHasChangedComponent)
-	{
-		//Remove MID we could have added to the old component
-		CleanupDistortion();
-
-		// If the component is changing from one camera actor to another camera actor, update the undistorted focal length to the focal length of the new component
-		if (AttachedComponent != nullptr)
-		{
-			if (UCineCameraComponent* const CineCameraComponent = Cast<UCineCameraComponent>(ActorComponent))
-			{
-				UndistortedFocalLength = CineCameraComponent->CurrentFocalLength;
-			}
-		}
-	}
-
 	Super::SetAttachedComponent(ActorComponent);
 
 	if (UCineCameraComponent* const CineCameraComponent = Cast<UCineCameraComponent>(AttachedComponent))
 	{	
-		// Initialize the most recent focal length to the current focal length of the camera to properly detect changes to this property
-		LastFocalLength = CineCameraComponent->CurrentFocalLength;
-
 		// Initialize the most recent filmback to the current filmback of the camera to properly detect changes to this property
 		LastFilmback = CineCameraComponent->Filmback;
 
-		//When our component has changed, make sure we update our 
-		//cached/original filmback to restore it correctly
-		//PIE case is odd because the camera is duplicated from editor and its filmback will already be distorted
-		if (bApplyDistortion == false 
-			|| (GetWorld() && GetWorld()->WorldType != EWorldType::PIE))
-		{
-			UpdateCachedFocalLength(CineCameraComponent);
-		}
-
-		UpdateDistortionHandler(CineCameraComponent);
-
 		ULensFile* SelectedLensFile = LensFilePicker.GetLensFile();
 		if (SelectedLensFile && SelectedLensFile->IsCineCameraCompatible(CineCameraComponent) == false)
 		{
@@ -156,8 +126,6 @@ void ULiveLinkCameraController::SetAttachedComponent(UActorComponent* ActorCompo
 
 void ULiveLinkCameraController::Cleanup()
 {
-	CleanupDistortion();
-
 	if (UCineCameraComponent* const CineCameraComponent = Cast<UCineCameraComponent>(AttachedComponent))
 	{
 		if (UCameraCalibrationSubsystem* SubSystem = GEngine->GetEngineSubsystem<UCameraCalibrationSubsystem>())
@@ -200,43 +168,7 @@ void ULiveLinkCameraController::PostEditChangeProperty(struct FPropertyChangedEv
 {
 	const FName PropertyName = (PropertyChangedEvent.Property != NULL) ? PropertyChangedEvent.Property->GetFName() : NAME_None;
 
-	if (PropertyName == GET_MEMBER_NAME_CHECKED(ULiveLinkCameraController, bApplyDistortion))
-	{
-		if (UCineCameraComponent* const CineCameraComponent = Cast<UCineCameraComponent>(AttachedComponent))
-		{
-			if (bApplyDistortion)
-			{
-				/** 
-				 * Cache filmback to be able to recover it once distortion is turned off
-				 * not part of the distortion setup since entering PIE will duplicate actor / components
-				 * and the duplicated one will already have modified filmback applied.
-				 */
-				UndistortedFocalLength = CineCameraComponent->CurrentFocalLength;
-				UE_LOG(LogLiveLinkCameraController, Verbose, TEXT("Enabling distortion. Cached focal length is %0.3fmm"), UndistortedFocalLength);
-
-				UpdateDistortionHandler(CineCameraComponent);
-			}
-			else
-			{
-				// Static lens model data may not be available because the controller may not be ticking, so query the subsystem for any distortion handler
-				UCameraCalibrationSubsystem* SubSystem = GEngine->GetEngineSubsystem<UCameraCalibrationSubsystem>();
-				FDistortionHandlerPicker DistortionHandlerPicker;
-				DistortionHandlerPicker.TargetCameraComponent = CineCameraComponent;
-				DistortionHandlerPicker.DistortionProducerID = DistortionProducerID;
-				ULensDistortionModelHandlerBase* Handler = SubSystem->FindDistortionModelHandler(DistortionHandlerPicker);
-
-				// Try to remove the post-process material from the cine camera's blendables. 
-				if (Handler)
-				{
-					CineCameraComponent->RemoveBlendable(Handler->GetDistortionMID());
-				}
-
-				const float OriginalFOV = FMath::RadiansToDegrees(2.0f * FMath::Atan(CineCameraComponent->Filmback.SensorWidth / (2.0f * UndistortedFocalLength)));
-				CineCameraComponent->SetFieldOfView(OriginalFOV);
-			}
-		}
-	}
-	else if (PropertyName == GET_MEMBER_NAME_CHECKED(ULiveLinkCameraController, bApplyNodalOffset))
+	if (PropertyName == GET_MEMBER_NAME_CHECKED(ULiveLinkCameraController, bApplyNodalOffset))
 	{
 		if (UCineCameraComponent* const CineCameraComponent = Cast<UCineCameraComponent>(AttachedComponent))
 		{
@@ -426,108 +358,28 @@ void ULiveLinkCameraController::ApplyNodalOffset(ULensFile* SelectedLensFile, UC
 	}
 }
 
-void ULiveLinkCameraController::ApplyDistortion(ULensFile* LensFile, UCineCameraComponent* CineCameraComponent)
+void ULiveLinkCameraController::ApplyDistortion(ULensFile* LensFile, UCineCameraComponent* CineCameraComponent, const FLiveLinkCameraStaticData* StaticData, const FLiveLinkCameraFrameData* FrameData)
 {
-	UMaterialInstanceDynamic* NewDistortionMID = nullptr;
-
-	UpdateCachedFocalLength(CineCameraComponent);
+	bool bCanUpdateDistortion = (StaticData->bIsFocusDistanceSupported || StaticData->bIsFocalLengthSupported || StaticData->bIsFieldOfViewSupported);
 
-	if (LensFile)
+	if (LensFile && bCanUpdateDistortion)
 	{
 		UCameraCalibrationSubsystem* const SubSystem = GEngine->GetEngineSubsystem<UCameraCalibrationSubsystem>();
 		const FString HandlerDisplayName = FString::Format(TEXT("{0} (Lens File)"), { LensFile->GetFName().ToString() });
 
 		FDistortionHandlerPicker DistortionHandlerPicker = { CineCameraComponent, DistortionProducerID, HandlerDisplayName };
 		LensDistortionHandler = SubSystem->FindOrCreateDistortionModelHandler(DistortionHandlerPicker, LensFile->LensInfo.LensModel);
-	}
 
-	if (LensDistortionHandler)
-	{
-		if (LensFile != nullptr)
+		if (LensDistortionHandler)
 		{
 			//Go through the lens file to get distortion data based on FIZ
 			//Our handler's displacement map will get updated
 			FDistortionData DistortionData;
 			const FVector2D CurrentSensorDimensions = FVector2D(CineCameraComponent->Filmback.SensorWidth, CineCameraComponent->Filmback.SensorHeight);
-			LensFile->EvaluateDistortionData(CineCameraComponent->CurrentFocusDistance, UndistortedFocalLength, CurrentSensorDimensions, LensDistortionHandler, DistortionData);
-		}
 
-		NewDistortionMID = LensDistortionHandler->GetDistortionMID();
-
-		if (bApplyDistortion)
-		{
-			// Scale the camera's FOV by an overscan factor
-			const float OverscanFactor = LensDistortionHandler->GetOverscanFactor();
-			const float OverscanSensorWidth = CineCameraComponent->Filmback.SensorWidth * OverscanFactor;
-			const float OverscanFOV = FMath::RadiansToDegrees(2.0f * FMath::Atan(OverscanSensorWidth / (2.0f * UndistortedFocalLength)));
-			CineCameraComponent->SetFieldOfView(OverscanFOV);
+			LensFile->EvaluateDistortionData(CineCameraComponent->CurrentFocusDistance, CineCameraComponent->CurrentFocalLength, CurrentSensorDimensions, LensDistortionHandler, DistortionData);
 		}
 	}
-
-	//Cleanup distortion MIDs we could have setup
-	if (NewDistortionMID != LastDistortionMID)
-	{
-		CleanupDistortion();
-	}
-	
-	//Stamp last MID used for distortion
-	LastDistortionMID = NewDistortionMID;
-
-	// If distortion should be applied to the attached cinecamera, fetch the distortion MID from the Lens Distortion Handler and add it to the camera's post-process materials
-	if (bApplyDistortion && (bIsDistortionSetup == false))
-	{
-		CineCameraComponent->AddOrUpdateBlendable(LastDistortionMID);
-		bIsDistortionSetup = true;
-	}
-	// If distortion should not be applied, remove the distortion MID from the camera's post process materials
-	else if (bIsDistortionSetup && (bApplyDistortion == false))
-	{
-		CleanupDistortion();
-	}
-
-	//Stamp last applied focal length to detect if user has changed it
-	LastFocalLength = CineCameraComponent->CurrentFocalLength;
-}
-
-void ULiveLinkCameraController::UpdateDistortionHandler(UCineCameraComponent* CineCameraComponent)
-{
- 	if (LensDistortionHandler)
- 	{
- 		//Cache MID when changing handler to start with something valid. i.e. Cleaning MID from blendables if LL was never valid
- 		LastDistortionMID = LensDistortionHandler->GetDistortionMID();
- 	}
-}
-
-void ULiveLinkCameraController::UpdateCachedFocalLength(UCineCameraComponent* CineCameraComponent)
-{
-	//Verify if focal length was changed by the user. If that's the case, take the current value and update the cached one
-	if (CineCameraComponent->CurrentFocalLength != LastFocalLength)
-	{
-		UndistortedFocalLength = CineCameraComponent->CurrentFocalLength;
-
-		UE_LOG(LogLiveLinkCameraController, Verbose, TEXT("Updating cached focal length to %0.3fmm"), UndistortedFocalLength);
-	}
-}
-
-void ULiveLinkCameraController::CleanupDistortion()
-{
-	//Remove MID we could have added to the component
-	if (bIsDistortionSetup)
-	{
-		if (UCineCameraComponent* const CineCameraComponent = Cast<UCineCameraComponent>(AttachedComponent))
-		{
-			CineCameraComponent->RemoveBlendable(LastDistortionMID);
-			
-			//Update cached filmback before resetting it. We could have stopped ticking because the LiveLink component was stopped
-			//In the meantime, filmback could have been manually changed and our cache is out of date.
-			UpdateCachedFocalLength(CineCameraComponent);
-
-			const float OriginalFOV = FMath::RadiansToDegrees(2.0f * FMath::Atan(CineCameraComponent->Filmback.SensorWidth / (2.0f * UndistortedFocalLength)));
-			CineCameraComponent->SetFieldOfView(OriginalFOV);
-		}
-	}
-
-	bIsDistortionSetup = false;
 }
 
 void ULiveLinkCameraController::OnPostActorTick(UWorld* World, ELevelTick TickType, float DeltaSeconds)
diff --git a/Engine/Plugins/VirtualProduction/LiveLinkCamera/Source/LiveLinkCamera/Public/LiveLinkCameraController.h b/Engine/Plugins/VirtualProduction/LiveLinkCamera/Source/LiveLinkCamera/Public/LiveLinkCameraController.h
--- a/Engine/Plugins/VirtualProduction/LiveLinkCamera/Source/LiveLinkCamera/Public/LiveLinkCameraController.h
+++ b/Engine/Plugins/VirtualProduction/LiveLinkCamera/Source/LiveLinkCamera/Public/LiveLinkCameraController.h
@@ -9,7 +9,6 @@
 #include "Engine/EngineTypes.h"
 #include "LensDistortionModelHandlerBase.h"
 #include "LensFile.h"
-#include "Materials/MaterialInstanceDynamic.h"
 
 #include "LiveLinkCameraController.generated.h"
 
@@ -60,5 +59,42 @@
 	GENERATED_BODY()
 
 public:
+	ULiveLinkCameraController();
+
+	//~ Begin ULiveLinkControllerBase interface
+	virtual void Tick(float DeltaTime, const FLiveLinkSubjectFrameData& SubjectData) override;
+	virtual bool IsRoleSupported(const TSubclassOf<ULiveLinkRole>& RoleToSupport) override;
+	virtual TSubclassOf<UActorComponent> GetDesiredComponentClass() const override;
+	virtual void SetAttachedComponent(UActorComponent* ActorComponent) override;
+	virtual void Cleanup() override;
+	virtual void OnEvaluateRegistered() override;
+	//~ End ULiveLinkControllerBase interface
+
+	//~ Begin UObject interface
+	virtual void PostDuplicate(bool bDuplicateForPIE) override;
+	virtual void PostEditImport() override;
+	virtual void PostLoad() override;
+#if WITH_EDITOR
+	virtual void PostEditChangeProperty(struct FPropertyChangedEvent& PropertyChangedEvent) override;
+#endif	
+	//~ End UObject interface
+
+	/** Returns true if encoder mapping is required for FIZ data */
+	bool IsEncoderMappingNeeded() const { return bIsEncoderMappingNeeded; }
+
+protected:
+	/** Applies FIZ data coming from LiveLink stream. Lens file is used if encoder mapping is required  */
+	void ApplyFIZ(ULensFile* LensFile, UCineCameraComponent* CineCameraComponent, const FLiveLinkCameraStaticData* StaticData, const FLiveLinkCameraFrameData* FrameData);
+
+	/** Applies nodal offset from lens file for the given Focus/Zoom values of CineCamera */
+	void ApplyNodalOffset(ULensFile* LensFile, UCineCameraComponent* CineCameraComponent);
+
+	/** Update distortion state */
+	void ApplyDistortion(ULensFile* LensFile, UCineCameraComponent* CineCameraComponent, const FLiveLinkCameraStaticData* StaticData, const FLiveLinkCameraFrameData* FrameData);
+
+	/** Verify base transform and apply nodal offset on top of everything else done in tick */
+	void OnPostActorTick(UWorld* World, ELevelTick TickType, float DeltaSeconds);
+
+public:
 #if WITH_EDITORONLY_DATA
 	UPROPERTY()
@@ -81,21 +117,13 @@
 	bool bApplyNodalOffset = true;
 
 protected:
-	/** Whether or not to apply a post-process distortion effect directly to the attached CineCamera */
-	UPROPERTY(EditAnywhere, Category = "Camera Calibration")
-	bool bApplyDistortion = false;
-
 	/** Cached distortion handler associated with attached camera component */
-	UPROPERTY(EditAnywhere, Category = "Camera Calibration", Transient)
-	ULensDistortionModelHandlerBase* LensDistortionHandler = nullptr;
-
-	/** Cached distortion MID the handler produced. Used to clean up old one in case it changes */
 	UPROPERTY(Transient)
-	UMaterialInstanceDynamic* LastDistortionMID = nullptr;
+	ULensDistortionModelHandlerBase* LensDistortionHandler = nullptr;
 
-	/** Original focal length settings of the attached cinecamera component to reapply when distortion isn't applied anymore */
-	UPROPERTY()
-	float UndistortedFocalLength = 50.0f;
+	/** Unique identifier representing the source of distortion data */
+	UPROPERTY(DuplicateTransient)
+	FGuid DistortionProducerID;
 
 	/** Original cinecamera component rotation that we set back on when nodal offset isn't applied anymore */
 	UPROPERTY()
@@ -105,86 +133,26 @@
 	UPROPERTY()
 	FVector OriginalCameraLocation;
 
-	/** Unique identifier representing the source of distortion data */
-	UPROPERTY(DuplicateTransient)
-	FGuid DistortionProducerID;
-
 	/** Used to control which data from LiveLink is actually applied to camera */
-	UPROPERTY(EditAnywhere, Category="Settings")
+	UPROPERTY(EditAnywhere, Category = "Settings")
 	FLiveLinkCameraControllerUpdateFlags UpdateFlags;
 
 #if WITH_EDITORONLY_DATA
 	/** Whether to refresh frustum drawing on value change */
-	UPROPERTY(EditAnywhere, Category="Debug")
+	UPROPERTY(EditAnywhere, Category = "Debug")
 	bool bShouldUpdateVisualComponentOnChange = true;
 #endif
 
 private:
-	
 	/** Whether incoming data requires encoder mapping */
 	bool bIsEncoderMappingNeeded = false;
 
-	/** Keep track of what needs to be setup to apply distortion */
-	bool bIsDistortionSetup = false;
-
 	/** Timestamp when we made the last warning log. Intervals to avoid log spamming */
 	double LastInvalidLoggingLoggedTimestamp = 0.0f;
 	static constexpr float TimeBetweenLoggingSeconds = 10.0f;
 
 	//Last values used to detect changes made by the user and update our original caches
-	float LastFocalLength = -1.0f;
 	FCameraFilmbackSettings LastFilmback;
 	FRotator LastRotation;
 	FVector LastLocation;
-
-public:
-
-	ULiveLinkCameraController();
-
-	//~ Begin ULiveLinkControllerBase interface
-	virtual void Tick(float DeltaTime, const FLiveLinkSubjectFrameData& SubjectData) override;
-	virtual bool IsRoleSupported(const TSubclassOf<ULiveLinkRole>& RoleToSupport) override;
-	virtual TSubclassOf<UActorComponent> GetDesiredComponentClass() const override;
-	virtual void SetAttachedComponent(UActorComponent* ActorComponent) override;
-	virtual void Cleanup() override;
-	virtual void OnEvaluateRegistered() override;
-	//~ End ULiveLinkControllerBase interface
-
-	//~ Begin UObject interface
-	virtual void PostDuplicate(bool bDuplicateForPIE) override;
-	virtual void PostEditImport() override;
-	virtual void PostLoad() override;
-#if WITH_EDITOR
-	virtual void PostEditChangeProperty(struct FPropertyChangedEvent& PropertyChangedEvent) override;
-#endif	
-	//~ End UObject interface
-
-	/** Returns true if encoder mapping is required for FIZ data */
-	bool IsEncoderMappingNeeded() const { return bIsEncoderMappingNeeded; }
-
-	/** Make sure what is required for distortion is setup or cleaned whether we apply or not distortion */
-	void UpdateDistortionSetup();
-
-protected:
-
-	/** Applies FIZ data coming from LiveLink stream. Lens file is used if encoder mapping is required  */
-	void ApplyFIZ(ULensFile* LensFile, UCineCameraComponent* CineCameraComponent, const FLiveLinkCameraStaticData* StaticData, const FLiveLinkCameraFrameData* FrameData);
-
-	/** Applies nodal offset from lens file for the given Focus/Zoom values of CineCamera */
-	void ApplyNodalOffset(ULensFile* LensFile, UCineCameraComponent* CineCameraComponent);
-
-	/** Update distortion state */
-	void ApplyDistortion(ULensFile* LensFile, UCineCameraComponent* CineCameraComponent);
-
-	/** Refresh distortion handler's, in case user deletes it  */
-	void UpdateDistortionHandler(UCineCameraComponent* CineCameraComponent);
-
-	/** Update cached focal length in case it was changed */
-	void UpdateCachedFocalLength(UCineCameraComponent* CineCameraComponent);
-
-	/** Cleanup distortion objects we could have added to camera */
-	void CleanupDistortion();
-
-	/** Verify base transform and apply nodal offset on top of everything else done in tick */
-	void OnPostActorTick(UWorld* World, ELevelTick TickType, float DeltaSeconds);
 };
diff --git a/Engine/Plugins/VirtualProduction/LiveLinkLens/Source/LiveLinkLens/Private/LiveLinkLensController.cpp b/Engine/Plugins/VirtualProduction/LiveLinkLens/Source/LiveLinkLens/Private/LiveLinkLensController.cpp
--- a/Engine/Plugins/VirtualProduction/LiveLinkLens/Source/LiveLinkLens/Private/LiveLinkLensController.cpp
+++ b/Engine/Plugins/VirtualProduction/LiveLinkLens/Source/LiveLinkLens/Private/LiveLinkLensController.cpp
@@ -3,16 +3,13 @@
 #include "LiveLinkLensController.h"
 
 #include "CameraCalibrationSubsystem.h"
+#include "CineCameraComponent.h"
 #include "Engine/Engine.h"
 #include "Engine/World.h"
-#include "LiveLinkComponentController.h"
 #include "LiveLinkLensRole.h"
 #include "LiveLinkLensTypes.h"
-#include "Logging/LogMacros.h"
 
 
-DEFINE_LOG_CATEGORY_STATIC(LogLiveLinkLensController, Log, All);
-
 ULiveLinkLensController::ULiveLinkLensController() 
 {
 	if (!HasAnyFlags(RF_ArchetypeObject | RF_ClassDefaultObject))
@@ -40,11 +37,6 @@ void ULiveLinkLensController::Tick(float DeltaTime, const FLiveLinkSubjectFrameD
 			FDistortionHandlerPicker DistortionHandlerPicker = { CineCameraComponent, DistortionProducerID, HandlerDisplayName };
 			LensDistortionHandler = SubSystem->FindOrCreateDistortionModelHandler(DistortionHandlerPicker, LensModel);
 			
-			//To keep track of an updated MID from the handler
-			UMaterialInstanceDynamic* NewDistortionMID = nullptr;
-
-			UpdateCachedFocalLength(CineCameraComponent);
-
 			if (LensDistortionHandler)
 			{
 				// Update the lens distortion handler with the latest frame of data from the LiveLink source
@@ -55,47 +47,14 @@ void ULiveLinkLensController::Tick(float DeltaTime, const FLiveLinkSubjectFrameD
 				DistortionState.PrincipalPoint = FrameData->PrincipalPoint;
 
 				//Update the distortion state based on incoming LL data.
-				//Recompute overscan factor for the distortion state
-				//Make sure the displacement map is up to date
 				LensDistortionHandler->SetDistortionState(DistortionState);
-				LensDistortionHandler->SetOverscanFactor(LensDistortionHandler->ComputeOverscanFactor());
-				LensDistortionHandler->ProcessCurrentDistortion();
-
-				NewDistortionMID = LensDistortionHandler->GetDistortionMID();
-
-				// Get the computed overscan factor and scale the camera's sensor dimensions to simulate a wider FOV
-				if (bApplyDistortion)
-				{
-					const float OverscanFactor = LensDistortionHandler->GetOverscanFactor();
-					const float OverscanSensorWidth = CineCameraComponent->Filmback.SensorWidth * OverscanFactor;
-					const float OverscanFOV = FMath::RadiansToDegrees(2.0f * FMath::Atan(OverscanSensorWidth / (2.0f * UndistortedFocalLength)));
-					CineCameraComponent->SetFieldOfView(OverscanFOV);
-				}
-			}
 
-			//Cleanup distortion MIDs we could have setup
-			if (NewDistortionMID != LastDistortionMID)
-			{
-				CleanupDistortion();
-			}
-
-			//Stamp last MID used for distortion
-			LastDistortionMID = NewDistortionMID;
+				//Recompute overscan factor for the distortion state
+				LensDistortionHandler->SetOverscanFactor(LensDistortionHandler->ComputeOverscanFactor());
 
-			// If distortion should be applied to the attached cinecamera, fetch the distortion MID from the Lens Distortion Handler and add it to the camera's post-process materials
-			if (bApplyDistortion && (bIsDistortionSetup == false))
-			{
-				CineCameraComponent->AddOrUpdateBlendable(LastDistortionMID);
-				bIsDistortionSetup = true;
-			}
-			// If distortion should not be applied, remove the distortion MID from the camera's post process materials
-			else if (bIsDistortionSetup && (bApplyDistortion == false))
-			{
-				CleanupDistortion();
+				//Make sure the displacement map is up to date
+				LensDistortionHandler->ProcessCurrentDistortion();
 			}
-
-			//Stamp last applied focal length to detect if user has changed it
-			LastFocalLength = CineCameraComponent->CurrentFocalLength;
 		}
 	}
 }
@@ -110,48 +69,8 @@ TSubclassOf<UActorComponent> ULiveLinkLensController::GetDesiredComponentClass()
 	return UCineCameraComponent::StaticClass();
 }
 
-void ULiveLinkLensController::SetAttachedComponent(UActorComponent* ActorComponent)
-{
-	const bool bHasChangedComponent = (ActorComponent != AttachedComponent);
-	if (bHasChangedComponent)
-	{
-		//Remove MID we could have added to the old component
-		CleanupDistortion();
-
-		// If the component is changing from one camera actor to another camera actor, update the undistorted focal length to the focal length of the new component
-		if (AttachedComponent != nullptr)
-		{
-			if (UCineCameraComponent* const CineCameraComponent = Cast<UCineCameraComponent>(ActorComponent))
-			{
-				UndistortedFocalLength = CineCameraComponent->CurrentFocalLength;
-			}
-		}
-	}
-
-	Super::SetAttachedComponent(ActorComponent);
-	
-	if (UCineCameraComponent* const CineCameraComponent = Cast<UCineCameraComponent>(AttachedComponent))
-	{
-		// Initialize the most recent focal length to the current focal length of the camera to properly detect changes to this property
-		LastFocalLength = CineCameraComponent->CurrentFocalLength;
-
-		//When our component has changed, make sure we update our 
-		//cached/original focal length to restore it correctly
-		//PIE case is odd because the camera is duplicated from editor and its focal length will already be distorted
-		if (bApplyDistortion == false
-			|| (GetWorld() && GetWorld()->WorldType != EWorldType::PIE))
-		{
-			UpdateCachedFocalLength(CineCameraComponent);
-		}
-
-		UpdateDistortionHandler(CineCameraComponent);
-	}
-}
-
 void ULiveLinkLensController::Cleanup()
 {
-	CleanupDistortion();
-
 	if (UCineCameraComponent* const CineCameraComponent = Cast<UCineCameraComponent>(AttachedComponent))
 	{
 		if (UCameraCalibrationSubsystem* SubSystem = GEngine->GetEngineSubsystem<UCameraCalibrationSubsystem>())
@@ -183,89 +102,3 @@ void ULiveLinkLensController::PostEditImport()
 		DistortionProducerID = FGuid::NewGuid();
 	}
 }
-
-#if WITH_EDITOR
-void ULiveLinkLensController::PostEditChangeProperty(struct FPropertyChangedEvent& PropertyChangedEvent)
-{
-	const FName PropertyName = (PropertyChangedEvent.Property != NULL) ? PropertyChangedEvent.Property->GetFName() : NAME_None;
-
-	if (PropertyName == GET_MEMBER_NAME_CHECKED(ULiveLinkLensController, bApplyDistortion))
-	{
-		if (UCineCameraComponent* const CineCameraComponent = Cast<UCineCameraComponent>(AttachedComponent))
-		{
-			if (bApplyDistortion)
-			{
-				/**
-				 * Cache focal length to be able to recover it once distortion is turned off
-				 * not part of the distortion setup since entering PIE will duplicate actor / components
-				 * and the duplicated one will already have modified focal length applied.
-				 */
-				UndistortedFocalLength = CineCameraComponent->CurrentFocalLength;
-				UE_LOG(LogLiveLinkLensController, Verbose, TEXT("Enabling distortion. Cached Focal Length is %0.3fmm"), UndistortedFocalLength);
-
-				UpdateDistortionHandler(CineCameraComponent);
-			}
-			else
-			{
-				// Static lens model data may not be available because the controller may not be ticking, so query the subsystem for any distortion handler
-				UCameraCalibrationSubsystem* SubSystem = GEngine->GetEngineSubsystem<UCameraCalibrationSubsystem>();
-				FDistortionHandlerPicker DistortionHandlerPicker;
-				DistortionHandlerPicker.TargetCameraComponent = CineCameraComponent;
-				DistortionHandlerPicker.DistortionProducerID = DistortionProducerID;
-				ULensDistortionModelHandlerBase* Handler = SubSystem->FindDistortionModelHandler(DistortionHandlerPicker);
-
-				// Try to remove the post-process material from the cine camera's blendables. 
-				if (Handler)
-				{
-					CineCameraComponent->RemoveBlendable(Handler->GetDistortionMID());
-				}
-
-				// Restore the original FOV of the 
-				const float OriginalFOV = FMath::RadiansToDegrees(2.0f * FMath::Atan(CineCameraComponent->Filmback.SensorWidth / (2.0f * UndistortedFocalLength)));
-				CineCameraComponent->SetFieldOfView(OriginalFOV);
-			}
-		}
-	}
-}
-#endif
-
-void ULiveLinkLensController::UpdateDistortionHandler(UCineCameraComponent* CineCameraComponent)
-{
-	if (LensDistortionHandler)
-	{
-		//Cache MID when changing handler to start with something valid. i.e. Cleaning MID from blendables if LL was never valid
-		LastDistortionMID = LensDistortionHandler->GetDistortionMID();
-	}
-}
-
-void ULiveLinkLensController::UpdateCachedFocalLength(UCineCameraComponent* CineCameraComponent)
-{
-	//Verify if focal length was changed by the user. If that's the case, take the current value and update the cached one
-	if (CineCameraComponent->CurrentFocalLength != LastFocalLength)
-	{
-		UndistortedFocalLength = CineCameraComponent->CurrentFocalLength;
-
-		UE_LOG(LogLiveLinkLensController, Verbose, TEXT("Updating cached focal length to %0.3fmm"), UndistortedFocalLength);
-	}
-}
-
-void ULiveLinkLensController::CleanupDistortion()
-{
-	//Remove MID we could have added to the component
-	if (bIsDistortionSetup)
-	{
-		if (UCineCameraComponent* const CineCameraComponent = Cast<UCineCameraComponent>(AttachedComponent))
-		{
-			CineCameraComponent->RemoveBlendable(LastDistortionMID);
-
-			//Update cached focal length before resetting it. We could have stopped ticking because the LiveLink component was stopped
-			//In the meantime, focal length could have been manually changed and our cache is out of date.
-			UpdateCachedFocalLength(CineCameraComponent);
-			
-			const float OriginalFOV = FMath::RadiansToDegrees(2.0f * FMath::Atan(CineCameraComponent->Filmback.SensorWidth / (2.0f * UndistortedFocalLength)));
-			CineCameraComponent->SetFieldOfView(OriginalFOV);
-		}
-	}
-
-	bIsDistortionSetup = false;
-}
diff --git a/Engine/Plugins/VirtualProduction/LiveLinkLens/Source/LiveLinkLens/Public/LiveLinkLensController.h b/Engine/Plugins/VirtualProduction/LiveLinkLens/Source/LiveLinkLens/Public/LiveLinkLensController.h
--- a/Engine/Plugins/VirtualProduction/LiveLinkLens/Source/LiveLinkLens/Public/LiveLinkLensController.h
+++ b/Engine/Plugins/VirtualProduction/LiveLinkLens/Source/LiveLinkLens/Public/LiveLinkLensController.h
@@ -4,9 +4,7 @@
 
 #include "LiveLinkControllerBase.h"
 
-#include "CineCameraComponent.h"
 #include "LensDistortionModelHandlerBase.h"
-#include "Materials/MaterialInstanceDynamic.h"
 
 #include "LiveLinkLensController.generated.h"
 
@@ -25,53 +23,20 @@ public:
 	virtual void Tick(float DeltaTime, const FLiveLinkSubjectFrameData& SubjectData) override;
 	virtual bool IsRoleSupported(const TSubclassOf<ULiveLinkRole>& RoleToSupport) override;
 	virtual TSubclassOf<UActorComponent> GetDesiredComponentClass() const override;
-	virtual void SetAttachedComponent(UActorComponent* ActorComponent) override;
 	virtual void Cleanup() override;
 	//~ End ULiveLinkControllerBase interface
 
 	//~ Begin UObject Interface
 	virtual void PostDuplicate(bool bDuplicateForPIE) override;
 	virtual void PostEditImport() override;
-#if WITH_EDITOR
-	virtual void PostEditChangeProperty(struct FPropertyChangedEvent& PropertyChangedEvent) override;
-#endif
 	//~ End UObject Interface
 
 protected:
-
-	/** Refresh distortion handler's, in case user deletes it  */
-	void UpdateDistortionHandler(UCineCameraComponent* CineCameraComponent);
-
-	/** Update cached focal length in case it was changed */
-	void UpdateCachedFocalLength(UCineCameraComponent* CineCameraComponent);
-
-	/** Cleanup distortion objects we could have added to camera */
-	void CleanupDistortion();
-
-protected:
-	/** Whether or not to apply a post-process distortion effect directly to the attached CineCamera */
-	UPROPERTY(EditAnywhere, Category = "Camera Calibration")
-	bool bApplyDistortion = false;
-
 	/** Cached distortion handler associated with attached camera component */
-	UPROPERTY(EditAnywhere, Category = "Camera Calibration", Transient)
-	ULensDistortionModelHandlerBase* LensDistortionHandler = nullptr;
-
-	/** Cached distortion MID the handler produced. Used to clean up old one in case it changes */
 	UPROPERTY(Transient)
-	UMaterialInstanceDynamic* LastDistortionMID = nullptr;
-
-	/** Original focal length of the attached cinecamera component to reapply when distortion isn't applied anymore */
- 	UPROPERTY()
-	float UndistortedFocalLength = 50.0f;
+	ULensDistortionModelHandlerBase* LensDistortionHandler = nullptr;
 
 	/** Unique identifier representing the source of distortion data */
 	UPROPERTY(DuplicateTransient)
 	FGuid DistortionProducerID;
-
-	/** Keep track of what needs to be setup to apply distortion */
-	bool bIsDistortionSetup = false;
-
-	/** Last value used to detect changes made by the user and update our original caches */
-	float LastFocalLength = -1.0f;
 };
\ No newline at end of file
