- Fix CIS (missing #if WITH_EDITOR guards)

diff --git a/Engine/Plugins/VirtualProduction/CameraCalibration/Source/CameraCalibration/Private/LensDistortionComponent.cpp b/Engine/Plugins/VirtualProduction/CameraCalibration/Source/CameraCalibration/Private/LensDistortionComponent.cpp
--- a/Engine/Plugins/VirtualProduction/CameraCalibration/Source/CameraCalibration/Private/LensDistortionComponent.cpp
+++ b/Engine/Plugins/VirtualProduction/CameraCalibration/Source/CameraCalibration/Private/LensDistortionComponent.cpp
@@ -126,6 +126,7 @@ void ULensDistortionComponent::TickComponent(float DeltaTime, ELevelTick TickTyp
 	Super::TickComponent(DeltaTime, TickType, ThisTickFunction);
 }
 
+#if WITH_EDITOR
 void ULensDistortionComponent::PostEditChangeProperty(struct FPropertyChangedEvent& PropertyChangedEvent)
 {
  	const FName PropertyName = (PropertyChangedEvent.Property != NULL) ? PropertyChangedEvent.Property->GetFName() : NAME_None;
@@ -209,6 +210,7 @@ void ULensDistortionComponent::PostEditChangeProperty(struct FPropertyChangedEve
 
 	Super::PostEditChangeProperty(PropertyChangedEvent);
 }
+#endif //WITH_EDITOR
 
 void ULensDistortionComponent::PostDuplicate(bool bDuplicateForPIE)
 {
