- Updating lens file editor and placeholder for calib tools

diff --git a/Engine/Plugins/VirtualProduction/CameraCalibration/CameraCalibration.uplugin b/Engine/Plugins/VirtualProduction/CameraCalibration/CameraCalibration.uplugin
--- a/Engine/Plugins/VirtualProduction/CameraCalibration/CameraCalibration.uplugin
+++ b/Engine/Plugins/VirtualProduction/CameraCalibration/CameraCalibration.uplugin
@@ -24,5 +24,11 @@
 			"Type": "Editor",
 			"LoadingPhase": "Default"
 		}
+	],
+	"Plugins": [
+		{
+			"Name": "LiveLink",
+			"Enabled": true
+		}
 	]
 }
\ No newline at end of file
diff --git a/Engine/Plugins/VirtualProduction/CameraCalibration/Source/CameraCalibration/Private/LensDistortionComponent.cpp b/Engine/Plugins/VirtualProduction/CameraCalibration/Source/CameraCalibration/Private/LensDistortionComponent.cpp
--- a/Engine/Plugins/VirtualProduction/CameraCalibration/Source/CameraCalibration/Private/LensDistortionComponent.cpp
+++ b/Engine/Plugins/VirtualProduction/CameraCalibration/Source/CameraCalibration/Private/LensDistortionComponent.cpp
@@ -4,6 +4,7 @@
 
 #include "CameraCalibrationSubsystem.h"
 #include "CineCameraComponent.h"
+#include "Engine/Engine.h"
 #include "LiveLinkComponentController.h"
 
 ULensDistortionComponent::ULensDistortionComponent()
diff --git a/Engine/Plugins/VirtualProduction/CameraCalibration/Source/CameraCalibration/Private/LensFile.cpp b/Engine/Plugins/VirtualProduction/CameraCalibration/Source/CameraCalibration/Private/LensFile.cpp
--- a/Engine/Plugins/VirtualProduction/CameraCalibration/Source/CameraCalibration/Private/LensFile.cpp
+++ b/Engine/Plugins/VirtualProduction/CameraCalibration/Source/CameraCalibration/Private/LensFile.cpp
@@ -628,51 +628,32 @@ bool ULensFile::EvaluateNodalPointOffset(float InFocus, float InZoom, FNodalPoin
 
 bool ULensFile::HasFocusEncoderMapping() const
 {
-	return EncoderMapping.Focus.Num() > 0;
+	return EncoderMapping.Focus.GetNumKeys() > 0;
 }
 
-bool ULensFile::EvaluateNormalizedFocus(float InNormalizedValue, float& OutEvaluatedValue)
+float ULensFile::EvaluateNormalizedFocus(float InNormalizedValue)
 {
-	/***TEMP TEMP TEMP*******/
-	/** Once there is a UI + methods to add encoder points we can get rid of that. */
-	TArray<FEncoderPoint> CopiedSorted = EncoderMapping.Focus;
-	CopiedSorted.Sort([](const FEncoderPoint& LHS, const FEncoderPoint& RHS) { return LHS.NormalizedValue <= RHS.NormalizedValue; });
-	/************************/
-
-	return LensInterpolationUtils::InterpolateEncoderValue(InNormalizedValue, CopiedSorted, OutEvaluatedValue);
+	return EncoderMapping.Focus.Eval(InNormalizedValue);
 }
 
 bool ULensFile::HasIrisEncoderMapping() const
 {
-	return EncoderMapping.Iris.Num() > 0;
-
+	return EncoderMapping.Iris.GetNumKeys() > 0;
 }
 
-float ULensFile::EvaluateNormalizedIris(float InNormalizedValue, float& OutEvaluatedValue)
+float ULensFile::EvaluateNormalizedIris(float InNormalizedValue)
 {
-	/***TEMP TEMP TEMP*******/
-	/** Once there is a UI + methods to add encoder points we can get rid of that. */
-	TArray<FEncoderPoint> CopiedSorted = EncoderMapping.Iris;
-	CopiedSorted.Sort([](const FEncoderPoint& LHS, const FEncoderPoint& RHS) { return LHS.NormalizedValue <= RHS.NormalizedValue; });
-	/************************/
-
-	return LensInterpolationUtils::InterpolateEncoderValue(InNormalizedValue, CopiedSorted, OutEvaluatedValue);
+	return EncoderMapping.Iris.Eval(InNormalizedValue);
 }
 
 bool ULensFile::HasZoomEncoderMapping() const
 {
-	return EncoderMapping.Zoom.Num() > 0;
+	return EncoderMapping.Zoom.GetNumKeys() > 0;
 }
 
-float ULensFile::EvaluateNormalizedZoom(float InNormalizedValue, float& OutEvaluatedValue)
+float ULensFile::EvaluateNormalizedZoom(float InNormalizedValue)
 {
-	/***TEMP TEMP TEMP*******/
-	/** Once there is a UI + methods to add encoder points we can get rid of that. */
-	TArray<FEncoderPoint> CopiedSorted = EncoderMapping.Zoom;
-	CopiedSorted.Sort([](const FEncoderPoint& LHS, const FEncoderPoint& RHS) { return LHS.NormalizedValue <= RHS.NormalizedValue; });
-	/************************/
-
-	return LensInterpolationUtils::InterpolateEncoderValue(InNormalizedValue, CopiedSorted, OutEvaluatedValue);
+	return EncoderMapping.Zoom.Eval(InNormalizedValue);
 }
 
 void ULensFile::OnDistortionDerivedDataJobCompleted(const FDerivedDistortionDataJobOutput& JobOutput)
@@ -799,3 +780,4 @@ ULensFile* FLensFilePicker::GetLensFile() const
 
 	return ReturnedLens;
 }
+
diff --git a/Engine/Plugins/VirtualProduction/CameraCalibration/Source/CameraCalibration/Private/LensInterpolationUtils.cpp b/Engine/Plugins/VirtualProduction/CameraCalibration/Source/CameraCalibration/Private/LensInterpolationUtils.cpp
--- a/Engine/Plugins/VirtualProduction/CameraCalibration/Source/CameraCalibration/Private/LensInterpolationUtils.cpp
+++ b/Engine/Plugins/VirtualProduction/CameraCalibration/Source/CameraCalibration/Private/LensInterpolationUtils.cpp
@@ -303,57 +303,6 @@ namespace LensInterpolationUtils
 		}
 	}
 
-	bool FindInterpolationIndex(float InRawValue, TArrayView<FEncoderPoint> InSourceData, int32& OutPointIndexA, int32& OutPointIndexB)
-	{
-		if (InSourceData.Num() <= 0)
-		{
-			return false;
-		}
-
-		for (int32 Index = InSourceData.Num() - 1; Index >= 0; --Index)
-		{
-			const FEncoderPoint& DataPoint = InSourceData[Index];
-			if (DataPoint.NormalizedValue <= InRawValue)
-			{
-				if (Index == InSourceData.Num() - 1)
-				{
-					OutPointIndexA = Index;
-					OutPointIndexB = Index;
-					return true;
-				}
-				else
-				{
-					OutPointIndexA = Index;
-					OutPointIndexB = Index + 1;
-					return true;
-				}
-			}
-		}
-
-		OutPointIndexA = 0;
-		OutPointIndexB = 0;
-		return true;
-	}
-
-	bool InterpolateEncoderValue(float InNormalizedValue, TArrayView<FEncoderPoint> InSourceData, float& OutEvaluatedValue)
-	{
-		int32 PointAIndex = 0;
-		int32 PointBIndex = 0;
-		if (FindInterpolationIndex(InNormalizedValue, InSourceData, PointAIndex, PointBIndex))
-		{
-			check(InSourceData.IsValidIndex(PointAIndex) && InSourceData.IsValidIndex(PointBIndex));
-
-			const FEncoderPoint& PointA = InSourceData[PointAIndex];
-			const FEncoderPoint& PointB = InSourceData[PointBIndex];
-
-			const float BlendingFactor = GetBlendFactor(InNormalizedValue, PointA.NormalizedValue, PointB.NormalizedValue);
-			OutEvaluatedValue = BlendValue(BlendingFactor, PointA.ValueInPhysicalUnits, PointB.ValueInPhysicalUnits);
-			return true;
-		}
-
-		return false;
-	}
-
 	float GetBlendFactor(float InValue, float ValueA, float ValueB)
 	{
 		//Keep input in range
diff --git a/Engine/Plugins/VirtualProduction/CameraCalibration/Source/CameraCalibration/Private/LensInterpolationUtils.h b/Engine/Plugins/VirtualProduction/CameraCalibration/Source/CameraCalibration/Private/LensInterpolationUtils.h
--- a/Engine/Plugins/VirtualProduction/CameraCalibration/Source/CameraCalibration/Private/LensInterpolationUtils.h
+++ b/Engine/Plugins/VirtualProduction/CameraCalibration/Source/CameraCalibration/Private/LensInterpolationUtils.h
@@ -9,9 +9,6 @@
 #include "UObject/ObjectMacros.h"
 
 
-struct FEncoderPoint;
-
-
 namespace LensInterpolationUtils
 {
 	//Flags for the possible direction on both axis
@@ -106,7 +103,6 @@ namespace LensInterpolationUtils
 	}
 	
 	float GetBlendFactor(float InValue, float ValueA, float ValueB);
-	bool InterpolateEncoderValue(float InNormalizedValue, TArrayView<FEncoderPoint> InSourceData, float& OutEvaluatedValue);
 
 	/** Helper function to find the possibly four indices for the points around the desired coordinates */
 	template<typename Type>
diff --git a/Engine/Plugins/VirtualProduction/CameraCalibration/Source/CameraCalibration/Private/Models/LensModel.cpp b/Engine/Plugins/VirtualProduction/CameraCalibration/Source/CameraCalibration/Private/Models/LensModel.cpp
--- a/Engine/Plugins/VirtualProduction/CameraCalibration/Source/CameraCalibration/Private/Models/LensModel.cpp
+++ b/Engine/Plugins/VirtualProduction/CameraCalibration/Source/CameraCalibration/Private/Models/LensModel.cpp
@@ -10,6 +10,30 @@
 
 DEFINE_LOG_CATEGORY_STATIC(LogLensModel, Log, All);
 
+#if WITH_EDITOR
+
+TArray<FText> ULensModel::GetParameterDisplayNames() const
+{
+	TArray<FText> ParameterNames;
+
+	UScriptStruct* TypeStruct = GetParameterStruct();
+
+	for (TFieldIterator<FProperty> It(TypeStruct); It; ++It)
+	{
+		if (FFloatProperty* Property = CastField<FFloatProperty>(*It))
+		{
+			ParameterNames.Add(Property->GetDisplayNameText());
+		}
+		else
+		{
+			UE_LOG(LogLensModel, Warning, TEXT("Property '%s' was skipped because its type was not float"), *(It->GetNameCPP()));
+		}
+	}
+
+	return ParameterNames;
+}
+
+#endif //WITH_EDITOR
 
 uint32 ULensModel::GetNumParameters() const
 {
@@ -18,7 +42,7 @@ uint32 ULensModel::GetNumParameters() const
 	uint32 NumParameters = 0;
 	for (TFieldIterator<FProperty> It(TypeStruct); It; ++It)
 	{
-		if (FFloatProperty* Prop = CastField<FFloatProperty>(*It))
+		if (FFloatProperty* Property = CastField<FFloatProperty>(*It))
 		{
 			++NumParameters;
 		}
@@ -50,7 +74,7 @@ TSubclassOf<ULensDistortionModelHandlerBase> ULensModel::GetHandlerClass(TSubcla
 	return nullptr;
 }
 
-void ULensModel::ToArray_Internal(UScriptStruct* TypeStruct, void* SrcData, TArray<float>& DstArray) const
+void ULensModel::ToArray_Internal(const UScriptStruct* TypeStruct, const void* SrcData, TArray<float>& DstArray) const
 {
 	if (TypeStruct != GetParameterStruct())
 	{
@@ -58,7 +82,7 @@ void ULensModel::ToArray_Internal(UScriptStruct* TypeStruct, void* SrcData, TArr
 		return;
 	}
 
-	DstArray.Reserve(GetNumParameters());
+	DstArray.Reset(GetNumParameters());
 	for (TFieldIterator<FProperty> It(TypeStruct); It; ++It)
 	{
 		if (FFloatProperty* Prop = CastField<FFloatProperty>(*It))
diff --git a/Engine/Plugins/VirtualProduction/CameraCalibration/Source/CameraCalibration/Public/LensData.h b/Engine/Plugins/VirtualProduction/CameraCalibration/Source/CameraCalibration/Public/LensData.h
--- a/Engine/Plugins/VirtualProduction/CameraCalibration/Source/CameraCalibration/Public/LensData.h
+++ b/Engine/Plugins/VirtualProduction/CameraCalibration/Source/CameraCalibration/Public/LensData.h
@@ -11,27 +11,6 @@
 
 #include "LensData.generated.h"
 
-USTRUCT(BlueprintType)
-struct CAMERACALIBRATION_API FEncoderPoint
-{
-	GENERATED_BODY()
-
-public:
-
-	/** Homed value in the range of [0..1] */
-	UPROPERTY(EditAnywhere, Category = "Point")
-	float NormalizedValue = 0.0f;
-
-	/*
-	 * Converted value in physical units
-	 * FIZ units
-	 * F: cm
-	 * I: F-Stops
-	 * Z: mm
-	 */
-	UPROPERTY(EditAnywhere, Category = "Point")
-	float ValueInPhysicalUnits = 0.0f;
-};
 
 /**
  * Information about the lens rig
diff --git a/Engine/Plugins/VirtualProduction/CameraCalibration/Source/CameraCalibration/Public/LensFile.h b/Engine/Plugins/VirtualProduction/CameraCalibration/Source/CameraCalibration/Public/LensFile.h
--- a/Engine/Plugins/VirtualProduction/CameraCalibration/Source/CameraCalibration/Public/LensFile.h
+++ b/Engine/Plugins/VirtualProduction/CameraCalibration/Source/CameraCalibration/Public/LensFile.h
@@ -4,6 +4,7 @@
 
 #include "CoreTypes.h"
 
+#include "Curves/RichCurve.h"
 #include "Engine/Texture.h"
 #include "ICalibratedMapProcessor.h"
 #include "LensData.h"
@@ -56,14 +57,20 @@ struct CAMERACALIBRATION_API FEncoderMapping
 
 public:
 
-	UPROPERTY(EditAnywhere, Category = "Encoder")
-	TArray<FEncoderPoint> Focus;
+	/** Focus curve from encoder values to nominal values */
+	UPROPERTY()
+	FRichCurve Focus;
 
-	UPROPERTY(EditAnywhere, Category = "Encoder")
-	TArray<FEncoderPoint> Iris;
+	/** Iris curve from encoder values to nominal values */
+	UPROPERTY()
+	FRichCurve Iris;
 
-	UPROPERTY(EditAnywhere, Category = "Encoder")
-	TArray<FEncoderPoint> Zoom;
+	/** 
+	 * Zoom curve from encoder values to nominal values 
+	 * @note To be removed and use only Fx/Fy calibrated values
+	 */
+	UPROPERTY()
+	FRichCurve Zoom;
 };
 
 /**
@@ -273,19 +280,19 @@ public:
 	bool HasFocusEncoderMapping() const;
 
 	/** Returns interpolated focus based on input normalized value and mapping */
-	bool EvaluateNormalizedFocus(float InNormalizedValue, float& OutEvaluatedValue);
+	float EvaluateNormalizedFocus(float InNormalizedValue);
 
 	/** Whether iris encoder mapping is configured */
 	bool HasIrisEncoderMapping() const;
 
 	/** Returns interpolated iris based on input normalized value and mapping */
-	float EvaluateNormalizedIris(float InNormalizedValue, float& OutEvaluatedValue);
+	float EvaluateNormalizedIris(float InNormalizedValue);
 
 	/** Whether zoom encoder mapping is configured */
 	bool HasZoomEncoderMapping() const;
 
 	/** Returns interpolated zoom based on input normalized value and mapping */
-	float EvaluateNormalizedZoom(float InNormalizedValue, float& OutEvaluatedValue);
+	float EvaluateNormalizedZoom(float InNormalizedValue);
 
 	/** Callbacked when stmap derived data has completed */
 	void OnDistortionDerivedDataJobCompleted(const FDerivedDistortionDataJobOutput& JobOutput);
diff --git a/Engine/Plugins/VirtualProduction/CameraCalibration/Source/CameraCalibration/Public/Models/LensModel.h b/Engine/Plugins/VirtualProduction/CameraCalibration/Source/CameraCalibration/Public/Models/LensModel.h
--- a/Engine/Plugins/VirtualProduction/CameraCalibration/Source/CameraCalibration/Public/Models/LensModel.h
+++ b/Engine/Plugins/VirtualProduction/CameraCalibration/Source/CameraCalibration/Public/Models/LensModel.h
@@ -4,6 +4,8 @@
 
 #include "CoreMinimal.h"
 
+#include "UObject/StructOnScope.h"
+
 #include "LensModel.generated.h"
 
 
@@ -25,6 +27,11 @@ public:
 	/** Get the struct of distortion parameters supported by this model */
 	virtual UScriptStruct* GetParameterStruct() const PURE_VIRTUAL(ULensModel::GetParameterStruct, return nullptr;);
 
+#if WITH_EDITOR
+	/** Get the names of each float parameters supported by this model */
+	virtual TArray<FText> GetParameterDisplayNames() const;
+#endif //WITH_EDITOR
+
 	/** Get the number of float fields in the parameter struct supported by this model */
 	virtual uint32 GetNumParameters() const;
 
@@ -33,11 +40,20 @@ public:
 	 * Note: the template type must be a UStruct
 	 */
 	template<typename StructType>
-	void ToArray(StructType& SrcData, TArray<float>& DstArray) const
+	void ToArray(const StructType& SrcData, TArray<float>& DstArray) const
 	{
 		ToArray_Internal(StructType::StaticStruct(), &SrcData, DstArray);
 	}
 
+	/** 
+	 * ToArray specialization taking a StructOnScope containing type and data 
+	 */
+	template<>
+	void ToArray<FStructOnScope>(const FStructOnScope& SrcData, TArray<float>& DstArray) const
+	{
+		ToArray_Internal(static_cast<const UScriptStruct*>(SrcData.GetStruct()), reinterpret_cast<const void*>(SrcData.GetStructMemory()), DstArray);
+	}
+
 	/**
 	 * Populate the float fields in the destination struct with the values in the source array
 	 * Note: the template type must be a UStruct
@@ -53,7 +69,7 @@ public:
 
 protected:
 	/** Internal implementation of ToArray. See declaration of public template method. */
-	virtual void ToArray_Internal(UScriptStruct* TypeStruct, void* SrcData, TArray<float>& DstArray) const;
+	virtual void ToArray_Internal(const UScriptStruct* TypeStruct, const void* SrcData, TArray<float>& DstArray) const;
 
 	/** Internal implementation of FromArray. See declaration of public template method. */
 	virtual void FromArray_Internal(UScriptStruct* TypeStruct, const TArray<float>& SrcArray, void* DstData);
diff --git a/Engine/Plugins/VirtualProduction/CameraCalibration/Source/CameraCalibrationEditor/CameraCalibrationEditor.Build.cs b/Engine/Plugins/VirtualProduction/CameraCalibration/Source/CameraCalibrationEditor/CameraCalibrationEditor.Build.cs
--- a/Engine/Plugins/VirtualProduction/CameraCalibration/Source/CameraCalibrationEditor/CameraCalibrationEditor.Build.cs
+++ b/Engine/Plugins/VirtualProduction/CameraCalibration/Source/CameraCalibrationEditor/CameraCalibrationEditor.Build.cs
@@ -14,9 +14,12 @@ namespace UnrealBuildTool.Rules
 					"CameraCalibration",
 					"Core",
 					"CoreUObject",
+	                "CurveEditor",
 					"EditorStyle",
 					"Engine",
 					"InputCore",
+					"LiveLinkEditor",
+					"LiveLinkInterface",
 					"PropertyEditor",
 					"Slate",
 					"SlateCore",
diff --git a/Engine/Plugins/VirtualProduction/CameraCalibration/Source/CameraCalibrationEditor/Private/AssetEditor/CameraCalibrationEditorCommon.h b/Engine/Plugins/VirtualProduction/CameraCalibration/Source/CameraCalibrationEditor/Private/AssetEditor/CameraCalibrationEditorCommon.h
new file mode 100644
--- /dev/null
+++ b/Engine/Plugins/VirtualProduction/CameraCalibration/Source/CameraCalibrationEditor/Private/AssetEditor/CameraCalibrationEditorCommon.h
@@ -0,0 +1,65 @@
+// Copyright Epic Games, Inc. All Rights Reserved.
+
+#pragma once
+
+#include "CoreMinimal.h"
+
+#include "Engine/Texture.h"
+#include "LensData.h"
+
+#include "CameraCalibrationEditorCommon.generated.h"
+
+
+/** Data categories manipulated in the camera calibration tools */
+enum class EDataCategories : uint8
+{
+	None,
+	Focus,
+	Iris,
+	Zoom,
+	Distortion,
+	ImageCenter,
+	STMap,
+	NodalOffset,
+};
+
+
+/** Container of distortion data to use instanced customization to show parameter struct defined by the model */
+USTRUCT()
+struct FDistortionInfoContainer
+{
+	GENERATED_BODY()
+
+public:
+
+	/** Distortion parameters */
+	UPROPERTY(EditAnywhere, Category = "Distortion", meta = (ShowOnlyInnerProperties))
+	FDistortionInfo DistortionInfo;
+};
+
+/**
+ * Intermediate structure to spawn stmap details
+ */
+USTRUCT()
+struct FCalibratedMapInfo
+{
+	GENERATED_BODY()
+
+public:
+
+	UPROPERTY(EditAnywhere, Category = "Distortion")
+	UTexture* DistortionMap = nullptr;
+};
+
+/** Holds the last FIZ data that was evaluated */
+struct FCachedFIZData
+{
+	TOptional<float> NormalizedFocus;
+	TOptional<float> NormalizedIris;
+	TOptional<float> NormalizedZoom;
+
+	TOptional<float> Focus;
+	TOptional<float> Iris;
+	TOptional<float> Zoom;
+};
+
diff --git a/Engine/Plugins/VirtualProduction/CameraCalibration/Source/CameraCalibrationEditor/Private/AssetEditor/CameraCalibrationToolkit.cpp b/Engine/Plugins/VirtualProduction/CameraCalibration/Source/CameraCalibrationEditor/Private/AssetEditor/CameraCalibrationToolkit.cpp
new file mode 100644
--- /dev/null
+++ b/Engine/Plugins/VirtualProduction/CameraCalibration/Source/CameraCalibrationEditor/Private/AssetEditor/CameraCalibrationToolkit.cpp
@@ -0,0 +1,220 @@
+// Copyright Epic Games, Inc. All Rights Reserved.
+
+#include "CameraCalibrationToolkit.h"
+
+#include "LensFile.h"
+#include "SLensEvaluation.h"
+#include "SLensFilePanel.h"
+#include "SNodalOffsetToolPanel.h"
+#include "UI/CameraCalibrationEditorStyle.h"
+#include "Widgets/Docking/SDockTab.h"
+
+
+#define LOCTEXT_NAMESPACE "CameraCalibrationToolkit"
+
+
+namespace CameraCalibrationToolkitUtils
+{
+	const FName CameraCalibrationIdentifier(TEXT("CameraCalibrationTools"));
+	const FName LensTabId(TEXT("LensFileEditorTab"));
+	const FName LensEvaluationTabId(TEXT("LensEvaluationTab"));
+	const FName NodalOffsetTabId(TEXT("NodalOffsetTab"));
+	const FName LensDetailsTabId(TEXT("LensFileDetailsTab"));
+}
+
+
+TSharedRef<FCameraCalibrationToolkit> FCameraCalibrationToolkit::CreateEditor(const EToolkitMode::Type Mode, const TSharedPtr<IToolkitHost>& InitToolkitHost, ULensFile* InLensFile)
+{
+	TSharedRef<FCameraCalibrationToolkit> NewEditor = MakeShared<FCameraCalibrationToolkit>();
+	NewEditor->InitCameraCalibrationTool(Mode, InitToolkitHost, InLensFile);
+
+	return NewEditor;
+}
+
+void FCameraCalibrationToolkit::InitCameraCalibrationTool(const EToolkitMode::Type Mode, const TSharedPtr<IToolkitHost>& InitToolkitHost, ULensFile* InLensFile)
+{
+	LensFile = InLensFile;
+	LensEvaluationWidget = SNew(SLensEvaluation, InLensFile);
+	NodalOffsetToolTab = SNew(SNodalOffsetToolPanel, LensFile);
+	LensEditorTab = SNew(SLensFilePanel, LensFile)
+		.CachedFIZData(TAttribute<FCachedFIZData>::Create(TAttribute<FCachedFIZData>::FGetter::CreateSP(LensEvaluationWidget.ToSharedRef(), &SLensEvaluation::GetLastEvaluatedData)));
+		
+	TSharedRef<FTabManager::FLayout> NewLayout = FTabManager::NewLayout("CameraCalibrationToolLayout_v0.6")
+		->AddArea
+		(
+			FTabManager::NewPrimaryArea()
+			->SetOrientation(Orient_Vertical)
+			->Split
+			(
+				FTabManager::NewStack()
+				->SetSizeCoefficient(0.1f)
+				->SetHideTabWell(true)
+				->AddTab(GetToolbarTabId(), ETabState::OpenedTab)
+			)
+			->Split
+			(
+				FTabManager::NewSplitter()
+				->SetSizeCoefficient(0.85f)
+				->SetOrientation(Orient_Horizontal)
+				->Split
+				(
+					FTabManager::NewStack()
+					->AddTab(CameraCalibrationToolkitUtils::LensTabId, ETabState::OpenedTab)
+					->AddTab(CameraCalibrationToolkitUtils::NodalOffsetTabId, ETabState::OpenedTab)
+					->AddTab(CameraCalibrationToolkitUtils::LensDetailsTabId, ETabState::ClosedTab)
+				)
+			)
+			->Split
+			(
+				FTabManager::NewStack()
+				->SetSizeCoefficient(0.15f)
+				->SetHideTabWell(true)
+				->AddTab(CameraCalibrationToolkitUtils::LensEvaluationTabId, ETabState::OpenedTab)
+			)
+		);
+
+	constexpr bool bCreateDefaultStandaloneMenu = true;
+	constexpr bool bCreateDefaultToolbar = true;
+	constexpr bool bToolbarFocusable = false;
+	constexpr bool bUseSmallIcons = true;
+	FAssetEditorToolkit::InitAssetEditor(
+		Mode,
+		InitToolkitHost,
+		CameraCalibrationToolkitUtils::CameraCalibrationIdentifier,
+		NewLayout,
+		bCreateDefaultStandaloneMenu,
+		bCreateDefaultToolbar,
+		InLensFile,
+		bToolbarFocusable,
+		bUseSmallIcons);
+
+	RegenerateMenusAndToolbars();
+}
+
+void FCameraCalibrationToolkit::RegisterTabSpawners(const TSharedRef<class FTabManager>& InTabManager)
+{
+	WorkspaceMenuCategory = InTabManager->AddLocalWorkspaceMenuCategory(LOCTEXT("WorkspaceMenu_CameraCalibrationTools", "Camera Calibration Panel"));
+
+	Super::RegisterTabSpawners(InTabManager);
+	
+	InTabManager->RegisterTabSpawner(CameraCalibrationToolkitUtils::LensTabId, FOnSpawnTab::CreateSP(this, &FCameraCalibrationToolkit::HandleSpawnLensEditorTab))
+		.SetDisplayName(LOCTEXT("LensEditorTab", "Lens File Editor"))
+		.SetGroup(WorkspaceMenuCategory.ToSharedRef())
+		.SetIcon(FSlateIcon(FEditorStyle::GetStyleSetName(), "LevelEditor.GameSettings.Small"));
+
+	InTabManager->RegisterTabSpawner(CameraCalibrationToolkitUtils::LensDetailsTabId, FOnSpawnTab::CreateSP(this, &FCameraCalibrationToolkit::HandleSpawnLensDetailsPanelTab))
+		.SetDisplayName(LOCTEXT("LensDetailsTab", "Lens DetailsPanel"))
+		.SetGroup(WorkspaceMenuCategory.ToSharedRef())
+		.SetIcon(FSlateIcon(FEditorStyle::GetStyleSetName(), "LevelEditor.GameSettings.Small"));
+	
+	InTabManager->RegisterTabSpawner(CameraCalibrationToolkitUtils::NodalOffsetTabId, FOnSpawnTab::CreateSP(this, &FCameraCalibrationToolkit::HandleSpawnNodalOffsetTab))
+		.SetDisplayName(LOCTEXT("NodalOffsetTab", "Nodal Offset"))
+		.SetGroup(WorkspaceMenuCategory.ToSharedRef())
+		.SetIcon(FSlateIcon(FEditorStyle::GetStyleSetName(), "LevelEditor.GameSettings.Small"));
+
+	InTabManager->RegisterTabSpawner(CameraCalibrationToolkitUtils::LensEvaluationTabId, FOnSpawnTab::CreateSP(this, &FCameraCalibrationToolkit::HandleSpawnLensEvaluationTab))
+		.SetDisplayName(LOCTEXT("NodalOffsetTab", "Nodal Offset"))
+		.SetGroup(WorkspaceMenuCategory.ToSharedRef())
+		.SetIcon(FSlateIcon(FEditorStyle::GetStyleSetName(), "LevelEditor.GameSettings.Small"));
+}
+
+void FCameraCalibrationToolkit::UnregisterTabSpawners(const TSharedRef<class FTabManager>& InTabManager)
+{
+	InTabManager->UnregisterTabSpawner(CameraCalibrationToolkitUtils::LensTabId);
+	InTabManager->UnregisterTabSpawner(CameraCalibrationToolkitUtils::NodalOffsetTabId);
+	InTabManager->UnregisterTabSpawner(CameraCalibrationToolkitUtils::LensDetailsTabId);
+	InTabManager->UnregisterTabSpawner(CameraCalibrationToolkitUtils::LensEvaluationTabId);
+	Super::UnregisterTabSpawners(InTabManager);
+}
+
+bool FCameraCalibrationToolkit::OnRequestClose()
+{
+	return true;
+}
+
+FText FCameraCalibrationToolkit::GetBaseToolkitName() const
+{
+	return LOCTEXT("PanelToolkitName", "Camera Calibration Tools");
+}
+
+FName FCameraCalibrationToolkit::GetToolkitFName() const
+{
+	static const FName PanelName("CameraCalibrationTools");
+	return PanelName;
+}
+
+FLinearColor FCameraCalibrationToolkit::GetWorldCentricTabColorScale() const
+{
+	return FLinearColor(1.0, 0.0f, 0.0f, 1.0f);
+}
+
+FString FCameraCalibrationToolkit::GetWorldCentricTabPrefix() const
+{
+	return LOCTEXT("CameraCalibrationTabPrefix", "CameraCalibrationTools").ToString();
+}
+
+TSharedRef<SDockTab> FCameraCalibrationToolkit::HandleSpawnNodalOffsetTab(const FSpawnTabArgs& Args) const
+{
+	check(Args.GetTabId() == CameraCalibrationToolkitUtils::NodalOffsetTabId);
+
+	return SNew(SDockTab)
+		.Label(LOCTEXT("NodalOffsetToolPanel", "Nodal Offset Tool"))
+		.TabColorScale(GetTabColorScale())
+		[
+			NodalOffsetToolTab.ToSharedRef()
+		];
+}
+
+TSharedRef<SDockTab> FCameraCalibrationToolkit::HandleSpawnLensEditorTab(const FSpawnTabArgs& Args) const
+{
+	check(Args.GetTabId() == CameraCalibrationToolkitUtils::LensTabId);
+
+	return SNew(SDockTab)
+		.Label(LOCTEXT("LensFilePanelLabel", "Lens File Panel"))
+		.TabColorScale(GetTabColorScale())
+		[
+			LensEditorTab.ToSharedRef()
+		];
+}
+
+TSharedRef<SDockTab> FCameraCalibrationToolkit::HandleSpawnLensDetailsPanelTab(const FSpawnTabArgs& Args) const
+{
+	check(Args.GetTabId() == CameraCalibrationToolkitUtils::LensDetailsTabId);
+
+	FDetailsViewArgs DetailsViewArg;
+	DetailsViewArg.bHideSelectionTip = true;
+	FPropertyEditorModule& PropertyModule = FModuleManager::LoadModuleChecked<FPropertyEditorModule>("PropertyEditor");
+	TSharedPtr<IDetailsView> DetailsView = PropertyModule.CreateDetailView(DetailsViewArg);
+	DetailsView->SetObject(LensFile);
+
+	return SNew(SDockTab)
+		.Label(LOCTEXT("LensDetailsViewTab", "Lens DetailsPanel"))
+		.TabColorScale(GetTabColorScale())
+		[
+			DetailsView.ToSharedRef()
+		];
+}
+
+FCachedFIZData FCameraCalibrationToolkit::GetFIZData() const
+{
+	if (LensEvaluationWidget)
+	{
+		return LensEvaluationWidget->GetLastEvaluatedData();
+	}
+
+	return FCachedFIZData();
+}
+
+TSharedRef<SDockTab> FCameraCalibrationToolkit::HandleSpawnLensEvaluationTab(const FSpawnTabArgs& Args) const
+{
+	check(Args.GetTabId() == CameraCalibrationToolkitUtils::LensEvaluationTabId);
+
+	return SNew(SDockTab)
+		.Label(LOCTEXT("LensEvaluationTabLabel", "Lens Evaluation"))
+		.TabColorScale(GetTabColorScale())
+		[
+			LensEvaluationWidget.ToSharedRef()
+		];
+}
+
+#undef LOCTEXT_NAMESPACE
\ No newline at end of file
diff --git a/Engine/Plugins/VirtualProduction/CameraCalibration/Source/CameraCalibrationEditor/Private/AssetEditor/CameraCalibrationToolkit.h b/Engine/Plugins/VirtualProduction/CameraCalibration/Source/CameraCalibrationEditor/Private/AssetEditor/CameraCalibrationToolkit.h
new file mode 100644
--- /dev/null
+++ b/Engine/Plugins/VirtualProduction/CameraCalibration/Source/CameraCalibrationEditor/Private/AssetEditor/CameraCalibrationToolkit.h
@@ -0,0 +1,78 @@
+// Copyright Epic Games, Inc. All Rights Reserved.
+
+#pragma once
+
+#include "Toolkits/AssetEditorToolkit.h"
+
+#include "CoreMinimal.h"
+#include "LensFile.h"
+#include "CameraCalibrationEditorCommon.h"
+
+
+class IToolkitHost;
+class SLensFilePanel;
+class SNodalOffsetToolPanel;
+class SLensEvaluation;
+class ULensFile;
+
+/** Toolkit to do camera calibration */
+class FCameraCalibrationToolkit : public FAssetEditorToolkit
+{
+	using Super = FAssetEditorToolkit;
+
+public:
+	/** Creates an editor for LensFile assets */
+	static TSharedRef<FCameraCalibrationToolkit> CreateEditor(const EToolkitMode::Type Mode, const TSharedPtr<IToolkitHost>& InitToolkitHost, ULensFile* InLensFile);
+
+	/**
+	 * Edits the specified table
+	 *
+	 * @param	Mode					Asset editing mode for this editor (standalone or world-centric)
+	 * @param	InitToolkitHost			When Mode is WorldCentric, this is the level editor instance to spawn this editor within
+	 * @param	InLensFile				The LensFile asset to edit
+	 */
+	void InitCameraCalibrationTool(const EToolkitMode::Type Mode, const TSharedPtr<IToolkitHost>& InitToolkitHost, ULensFile* InLensFile);
+
+protected:
+
+	//~ Begin IToolkit interface
+	virtual FText GetBaseToolkitName() const override;
+	virtual FName GetToolkitFName() const override;
+	virtual FLinearColor GetWorldCentricTabColorScale() const override;
+	virtual FString GetWorldCentricTabPrefix() const override;
+	virtual void RegisterTabSpawners(const TSharedRef<FTabManager>& TabManager) override;
+	virtual void UnregisterTabSpawners(const TSharedRef<FTabManager>& TabManager) override;
+	virtual bool OnRequestClose() override;
+	//~ End IToolkit interface
+
+private:
+
+	/** Handle spawning the tab that holds the LensFile panel tab. */
+	TSharedRef<SDockTab> HandleSpawnLensEditorTab(const FSpawnTabArgs& Args) const;
+
+	/** Handle spawning the tab that holds the NodalOffset panel tab. */
+	TSharedRef<SDockTab> HandleSpawnNodalOffsetTab(const FSpawnTabArgs& Args) const;
+	
+	/** Handle spawning the tab that holds the current evaluated data */
+	TSharedRef<SDockTab> HandleSpawnLensEvaluationTab(const FSpawnTabArgs& Args) const;
+
+	/** Temporary tab to view the details panel of the lens file */
+	TSharedRef<SDockTab> HandleSpawnLensDetailsPanelTab(const FSpawnTabArgs& Args) const;
+
+	FCachedFIZData GetFIZData() const;
+
+
+private:
+
+	/** Toolkit is adding a reference to the editing object so a standard pointer is ok */
+	ULensFile* LensFile = nullptr;
+
+	/** Lens editor tab to do manual adjustments */
+	TSharedPtr<SLensFilePanel> LensEditorTab;
+
+	/** Placeholder tab to have camera calibration tools */
+	TSharedPtr<SNodalOffsetToolPanel> NodalOffsetToolTab;
+
+	/** Data displaying evaluated data based on LiveLink inputs */
+	TSharedPtr<SLensEvaluation> LensEvaluationWidget;
+};
\ No newline at end of file
diff --git a/Engine/Plugins/VirtualProduction/CameraCalibration/Source/CameraCalibrationEditor/Private/AssetEditor/LensFileEditorToolkit.cpp b/Engine/Plugins/VirtualProduction/CameraCalibration/Source/CameraCalibrationEditor/Private/AssetEditor/LensFileEditorToolkit.cpp
deleted file mode 100644
diff --git a/Engine/Plugins/VirtualProduction/CameraCalibration/Source/CameraCalibrationEditor/Private/AssetEditor/LensFileEditorToolkit.h b/Engine/Plugins/VirtualProduction/CameraCalibration/Source/CameraCalibrationEditor/Private/AssetEditor/LensFileEditorToolkit.h
deleted file mode 100644
diff --git a/Engine/Plugins/VirtualProduction/CameraCalibration/Source/CameraCalibrationEditor/Private/AssetEditor/SLensDataPointDialog.cpp b/Engine/Plugins/VirtualProduction/CameraCalibration/Source/CameraCalibrationEditor/Private/AssetEditor/SLensDataPointDialog.cpp
new file mode 100644
--- /dev/null
+++ b/Engine/Plugins/VirtualProduction/CameraCalibration/Source/CameraCalibrationEditor/Private/AssetEditor/SLensDataPointDialog.cpp
@@ -0,0 +1,639 @@
+// Copyright Epic Games, Inc. All Rights Reserved.
+
+#include "SLensDataPointDialog.h"
+
+#include "Curves/RichCurve.h"
+#include "EditorStyleSet.h"
+#include "Framework/MultiBox/MultiBoxBuilder.h"
+#include "IDetailChildrenBuilder.h"
+#include "IDetailPropertyRow.h"
+#include "ISinglePropertyView.h"
+#include "IStructureDetailsView.h"
+#include "LensFile.h"
+#include "Models/LensModel.h"
+#include "Modules/ModuleManager.h"
+#include "PropertyHandle.h"
+#include "Widgets/Input/SButton.h"
+#include "Widgets/Input/SCheckBox.h"
+#include "Widgets/Input/SNumericEntryBox.h"
+
+
+#define LOCTEXT_NAMESPACE "LensDataPointDialog"
+
+TWeakPtr<SWindow> SLensDataPointDialog::ExistingWindow;
+
+/**
+ * Instanced customization used to display the distortion parameters
+ * data structure associated with selected model
+ */
+class FDistortionInfoCustomization : public IPropertyTypeCustomization
+{
+public:
+
+	FDistortionInfoCustomization(TSubclassOf<ULensModel> InLensModel)
+		: IPropertyTypeCustomization()
+		, CurrentLensModel(MoveTemp(InLensModel))
+	{}
+
+	// IDetailCustomization interface
+	virtual void CustomizeHeader(TSharedRef<IPropertyHandle> PropertyHandle, FDetailWidgetRow& HeaderRow, IPropertyTypeCustomizationUtils& CustomizationUtils) override
+	{
+	}
+
+	virtual void CustomizeChildren(TSharedRef<IPropertyHandle> StructPropertyHandle, IDetailChildrenBuilder& StructBuilder, IPropertyTypeCustomizationUtils& CustomizationUtils) override
+	{
+		if (CurrentLensModel)
+		{
+			FStructureDetailsViewArgs StructureViewArgs;
+			FDetailsViewArgs DetailArgs;
+			DetailArgs.bAllowSearch = false;
+			DetailArgs.bShowScrollBar = true;
+
+			FPropertyEditorModule& PropertyEditor = FModuleManager::Get().LoadModuleChecked<FPropertyEditorModule>(TEXT("PropertyEditor"));
+
+			DistortionParameters = MakeShared<FStructOnScope>(CurrentLensModel.GetDefaultObject()->GetParameterStruct());
+			TSharedPtr<IStructureDetailsView> StructureDetailsView = PropertyEditor.CreateStructureDetailView(DetailArgs, StructureViewArgs, DistortionParameters);
+
+			uint32 NumChildren;
+			StructPropertyHandle->GetNumChildren(NumChildren);
+
+			for (uint32 ChildIndex = 0; ChildIndex < NumChildren; ++ChildIndex)
+			{
+				const TSharedRef<IPropertyHandle> ChildHandle = StructPropertyHandle->GetChildHandle(ChildIndex).ToSharedRef();
+				const FName PropertyName = ChildHandle->GetProperty()->GetFName();
+
+				if (PropertyName == GET_MEMBER_NAME_CHECKED(FDistortionInfo, Parameters))
+				{
+					DistortionParametersArrayHandle = ChildHandle;
+					IDetailPropertyRow* ParametersRow = StructBuilder.AddExternalStructure(DistortionParameters.ToSharedRef());
+					ParametersRow->GetPropertyHandle()->SetOnChildPropertyValueChanged(FSimpleDelegate::CreateSP(this, &FDistortionInfoCustomization::OnDistortionParametersChanged));
+				}
+				else
+				{
+					StructBuilder.AddProperty(ChildHandle).ShowPropertyButtons(false);
+				}
+			}
+
+			//Start with initial parameters value
+			OnDistortionParametersChanged();
+		}
+	}
+
+private:
+
+	void OnDistortionParametersChanged() const
+	{
+		//Whenever a parameter value has changed, update the associated array
+		if (!CurrentLensModel)
+		{
+			return;
+		}
+
+		TArray<void*> RawData;
+		DistortionParametersArrayHandle->AccessRawData(RawData);
+		if (RawData.Num() == 1 && RawData[0])
+		{
+			TArray<float>* RawParametersArray = reinterpret_cast<TArray<float>*>(RawData[0]);
+			CurrentLensModel.GetDefaultObject()->ToArray(*DistortionParameters, *RawParametersArray);
+		}
+	}
+
+private:
+
+	/** Handle to the raw distortion parameters array */
+	TSharedPtr<IPropertyHandle> DistortionParametersArrayHandle;
+
+	/** Holds named parameters data structure */
+	TSharedPtr<FStructOnScope> DistortionParameters;
+
+	/** LensModel being edited */
+	TSubclassOf<ULensModel> CurrentLensModel;
+};
+
+void SLensDataPointDialog::Construct(const FArguments& InArgs, ULensFile* InLensFile)
+{
+	LensFile = TStrongObjectPtr<ULensFile>(InLensFile);
+	CachedFIZ = InArgs._CachedFIZData;
+	OnDataPointAdded = InArgs._OnDataPointAdded;
+
+	ChildSlot
+	[
+		SNew(SVerticalBox)
+		+ SVerticalBox::Slot()
+		.Padding(5.0f, 5.0f)
+		.AutoHeight()
+		[
+			SNew(SBorder)
+			.HAlign(HAlign_Fill)
+			.VAlign(VAlign_Fill)
+			.Padding(FMargin(4.f, 4.f, 4.f, 4.f))
+			[
+				SNew(SHorizontalBox)
+				+ SHorizontalBox::Slot()
+				.AutoWidth()
+				[
+					SNew(STextBlock)
+					.Text(LOCTEXT("DataCategoryLabel", "Data Type "))
+				]
+				+ SHorizontalBox::Slot()
+				.HAlign(HAlign_Fill)
+				[
+					SNew(STextBlock)
+					.Text(MakeAttributeLambda([=]
+					{
+							switch (SelectedCategory)
+							{
+								case EDataCategories::Focus: return LOCTEXT("FocusCategory", "Encoder - Focus");
+								case EDataCategories::Iris: return LOCTEXT("IrisCategory", "Encoder - Iris");
+								case EDataCategories::Zoom: return LOCTEXT("ZoomCategory", "Encoder - Zoom");
+								case EDataCategories::Distortion: return LOCTEXT("DistortionCategory", "Distortion Parameters");
+								case EDataCategories::ImageCenter: return LOCTEXT("ImageCenterCategory", "Image Center");
+								case EDataCategories::NodalOffset: return LOCTEXT("NodalOffsetCategory", "Nodal offset");
+								case EDataCategories::STMap: return LOCTEXT("STMapCategory", "STMap");
+								default: return LOCTEXT("InvalidCategory", "Invalid");
+							}
+					}))
+				]
+				+ SHorizontalBox::Slot()
+				.AutoWidth()
+				.Padding(FMargin(4.0f, 0.0f, 0.0f, 0.0f))
+				[
+					SNew(SComboButton)
+					.OnGetMenuContent(this, &SLensDataPointDialog::MakeDataCategoryMenuWidget)
+					.ContentPadding(FMargin(4.0, 2.0))
+				]
+			]
+		]
+		+ SVerticalBox::Slot()
+		.Padding(5.0f, 5.0f)
+		.AutoHeight()
+		[
+			SAssignNew(TrackingDataContainer, SBorder)
+			.HAlign(HAlign_Fill)
+			.VAlign(VAlign_Fill)
+			.Padding(FMargin(4.f, 4.f, 4.f, 4.f))
+		]
+		+ SVerticalBox::Slot()
+		.Padding(5.0f, 5.0f)
+		.FillHeight(0.7f)
+		[
+			SAssignNew(LensDataContainer, SBorder)
+			.HAlign(HAlign_Fill)
+			.VAlign(VAlign_Fill)
+			.Padding(FMargin(4.f, 4.f, 4.f, 4.f))
+		]
+		+ SVerticalBox::Slot()
+		.Padding(5.0f, 5.0f)
+		.AutoHeight()
+		[
+			SNew(SBorder)
+			.HAlign(HAlign_Fill)
+			.VAlign(VAlign_Fill)
+			.Padding(FMargin(4.f, 4.f, 4.f, 4.f))
+			[
+				MakeButtonsWidget()
+			]
+		]
+	];
+	
+	SetDataCategory(EDataCategories::Focus);
+}
+
+void SLensDataPointDialog::Tick(const FGeometry& AllottedGeometry, const double InCurrentTime, const float InDeltaTime)
+{
+	RefreshEvaluationData();
+	Super::Tick(AllottedGeometry, InCurrentTime, InDeltaTime);
+}
+
+void SLensDataPointDialog::OpenDialog(ULensFile* InLensFile, TAttribute<FCachedFIZData> InCachedFIZData, const FSimpleDelegate& InOnDataPointAdded)
+{
+	TSharedPtr<SWindow> ExistingWindowPin = ExistingWindow.Pin();
+	if (ExistingWindowPin.IsValid())
+	{
+		ExistingWindowPin->BringToFront();
+	}
+	else
+	{
+		ExistingWindowPin = SNew(SWindow)
+			.Title(LOCTEXT("LensEditorAddPointDialog", "Add Lens Data Point"))
+			.HasCloseButton(true)
+			.SupportsMaximize(false)
+			.SupportsMinimize(false)
+			.ClientSize(FVector2D(480, 360));
+
+		FSlateApplication::Get().AddWindow(ExistingWindowPin.ToSharedRef());
+	}
+
+	TSharedPtr<SLensDataPointDialog> AddPointDialog = 
+		SNew(SLensDataPointDialog, InLensFile)
+		.CachedFIZData(MoveTemp(InCachedFIZData))
+		.OnDataPointAdded(InOnDataPointAdded);
+
+	ExistingWindowPin->SetContent(AddPointDialog.ToSharedRef());
+	ExistingWindow = ExistingWindowPin;
+}
+
+FReply SLensDataPointDialog::OnAddDataPointClicked()
+{
+	AddDataToLensFile();
+
+	OnDataPointAdded.ExecuteIfBound();
+
+	CloseDialog();
+	return FReply::Handled();
+}
+
+FReply SLensDataPointDialog::OnCancelDataPointClicked()
+{
+	CloseDialog();
+	return FReply::Handled();
+}
+
+TSharedRef<SWidget> SLensDataPointDialog::MakeTrackingDataWidget()
+{
+	TSharedPtr<SWidget> TrackingWidget = SNullWidget::NullWidget;
+
+	const auto MakeRowWidget = [this](const FString& Label, int32 Index) -> TSharedRef<SWidget>
+	{
+		return SNew(SHorizontalBox)
+			+ SHorizontalBox::Slot()
+			.AutoWidth()
+			.VAlign(VAlign_Center)
+			[
+				SNew(SCheckBox)
+				.IsChecked(MakeAttributeLambda([this, Index]() { return IsTrackingDataOverrided(Index) ? ECheckBoxState::Checked : ECheckBoxState::Unchecked; }))
+				.OnCheckStateChanged(this, &SLensDataPointDialog::OnOverrideTrackingData, Index)
+			]
+			+ SHorizontalBox::Slot()
+			.FillWidth(0.2f)
+			[
+				SNew(STextBlock)
+				.Text(FText::FromString(Label))
+			]
+			+ SHorizontalBox::Slot()
+			.FillWidth(0.8f)
+			[
+				SNew(SNumericEntryBox<float>)
+				.IsEnabled(MakeAttributeLambda([this, Index]() { return IsTrackingDataOverrided(Index); }))
+				.Value(MakeAttributeLambda([this, Index]() { return TOptional<float>(GetTrackingData(Index)); }))
+				.OnValueChanged(this, &SLensDataPointDialog::SetTrackingData, Index)
+			];
+	};
+
+	//Based on category, either have one tracking input or two
+	switch (SelectedCategory)
+	{
+		case EDataCategories::Focus:
+		{
+			TrackingWidget = MakeRowWidget(TEXT("Input Focus"), 0);
+			break;
+		}
+		case EDataCategories::Iris:
+		{
+			TrackingWidget = MakeRowWidget(TEXT("Input Iris"), 0);
+			break;
+		}
+		case EDataCategories::Zoom:
+		{
+			TrackingWidget = MakeRowWidget(TEXT("Input Zoom"), 0);
+			break;
+		}
+		case EDataCategories::Distortion:
+		case EDataCategories::ImageCenter:
+		case EDataCategories::NodalOffset:
+		case EDataCategories::STMap:
+		default:
+		{
+			TrackingWidget = 
+				SNew(SVerticalBox)
+				+ SVerticalBox::Slot()
+				.AutoHeight()
+				[
+					MakeRowWidget(TEXT("Input Focus"), 0)
+				]
+				+ SVerticalBox::Slot()
+				.AutoHeight()
+				[
+					MakeRowWidget(TEXT("Input Zoom"), 1)
+				];
+			break;
+		}	
+	}
+
+	return TrackingWidget.ToSharedRef();
+}
+
+TSharedRef<SWidget> SLensDataPointDialog::MakeLensDataWidget()
+{
+	FStructureDetailsViewArgs StructureViewArgs;
+	FDetailsViewArgs DetailArgs;
+	DetailArgs.bAllowSearch = false;
+	DetailArgs.bShowScrollBar = false;
+
+	FPropertyEditorModule& PropertyEditor = FModuleManager::Get().LoadModuleChecked<FPropertyEditorModule>(TEXT("PropertyEditor"));
+	TSharedPtr<SWidget> LensDataWidget = SNullWidget::NullWidget;
+	
+	TSharedPtr<IStructureDetailsView> StructureDetailsView = PropertyEditor.CreateStructureDetailView(DetailArgs, StructureViewArgs, TSharedPtr<FStructOnScope>());
+	
+	LensData.Reset();
+	
+	switch (SelectedCategory)
+	{
+		case EDataCategories::Focus:
+		case EDataCategories::Iris:
+		case EDataCategories::Zoom:
+		{
+			LensDataWidget = MakeEncoderMappingWidget();
+			break;
+		}
+		case EDataCategories::Distortion:
+		{
+			LensData = MakeShared<FStructOnScope>(FDistortionInfoContainer::StaticStruct());
+			StructureDetailsView->GetDetailsView()->RegisterInstancedCustomPropertyTypeLayout(FDistortionInfo::StaticStruct()->GetFName(),
+				FOnGetPropertyTypeCustomizationInstance::CreateLambda([this]() { return MakeShared<FDistortionInfoCustomization>(LensFile->LensInfo.LensModel); })
+			);
+
+			StructureDetailsView->SetStructureData(LensData);
+			LensDataWidget = StructureDetailsView->GetWidget();
+
+			break;
+		}
+		case EDataCategories::ImageCenter:
+		{
+			LensData = MakeShared<FStructOnScope>(FIntrinsicParameters::StaticStruct());
+			StructureDetailsView->SetStructureData(LensData);
+			LensDataWidget = StructureDetailsView->GetWidget();
+			break;
+		}
+		case EDataCategories::NodalOffset:
+		{
+			LensData = MakeShared<FStructOnScope>(FNodalPointOffset::StaticStruct());
+			StructureDetailsView->SetStructureData(LensData);
+			LensDataWidget = StructureDetailsView->GetWidget();
+			break;
+		}
+		case EDataCategories::STMap:
+		{
+			LensData = MakeShared<FStructOnScope>(FCalibratedMapInfo::StaticStruct());
+			StructureDetailsView->SetStructureData(LensData);
+			LensDataWidget = StructureDetailsView->GetWidget();
+			break;
+		}
+		default:
+		{
+			break;
+		}
+	}
+
+	return LensDataWidget.ToSharedRef();
+}
+
+TSharedRef<SWidget> SLensDataPointDialog::MakeButtonsWidget()
+{
+	return SNew(SHorizontalBox)
+		+ SHorizontalBox::Slot()
+		[
+			SNew(SButton)
+			.ContentPadding(FEditorStyle::GetMargin("StandardDialog.ContentPadding"))
+			.OnClicked(this, &SLensDataPointDialog::OnAddDataPointClicked)
+			.HAlign(HAlign_Center)
+			.Text(LOCTEXT("AddDataPoint", "Add"))
+		]
+		+ SHorizontalBox::Slot()
+		[
+			SNew(SButton)
+			.ContentPadding(FEditorStyle::GetMargin("StandardDialog.ContentPadding"))
+			.OnClicked(this, &SLensDataPointDialog::OnCancelDataPointClicked)
+			.HAlign(HAlign_Center)
+			.Text(LOCTEXT("CancelAddingDataPoint", "Cancel"))
+		];
+}
+
+TSharedRef<SWidget> SLensDataPointDialog::MakeDataCategoryMenuWidget()
+{
+	constexpr bool bShouldCloseWindowAfterClosing = true;
+	FMenuBuilder MenuBuilder(bShouldCloseWindowAfterClosing, nullptr);
+
+	MenuBuilder.BeginSection("LensDataCategories", LOCTEXT("LensDataCategories", "Lens data categories"));
+	{
+		const auto AddMenuEntryFn = [&MenuBuilder, this](const FString& Label, EDataCategories Category)
+		{
+			MenuBuilder.AddMenuEntry
+			(
+				FText::FromString(Label)
+				, FText::FromString(Label)
+				, FSlateIcon()
+				, FUIAction
+				(
+					FExecuteAction::CreateLambda([this, Category] { SetDataCategory(Category); })
+					, FCanExecuteAction()
+					, FIsActionChecked::CreateLambda([this, Category] { return SelectedCategory == Category; })
+				)
+				, NAME_None
+				, EUserInterfaceActionType::RadioButton
+			);
+		};
+		
+		//Add different categories that we support
+		AddMenuEntryFn(TEXT("Focus"), EDataCategories::Focus);
+		AddMenuEntryFn(TEXT("Iris"), EDataCategories::Iris);
+		AddMenuEntryFn(TEXT("Distortion Parameters"), EDataCategories::Distortion);
+		AddMenuEntryFn(TEXT("Image Center"), EDataCategories::ImageCenter);
+		AddMenuEntryFn(TEXT("Nodal Offset"), EDataCategories::NodalOffset);
+		AddMenuEntryFn(TEXT("STMap"), EDataCategories::STMap);
+	}
+	MenuBuilder.EndSection();
+
+	return MenuBuilder.MakeWidget();
+}
+
+TSharedRef<SWidget> SLensDataPointDialog::MakeEncoderMappingWidget()
+{
+	return 
+		SNew(SVerticalBox)
+		+ SVerticalBox::Slot()
+		.AutoHeight()
+		[
+			SNew(SSplitter)
+			.Style(FEditorStyle::Get(), "DetailsView.Splitter")
+			.PhysicalSplitterHandleSize(1.0f)
+			.HitDetectionSplitterHandleSize(5.0f)
+			+ SSplitter::Slot()
+			[
+				SNew(SHorizontalBox)
+				+ SHorizontalBox::Slot()
+				.Padding(FMargin(0, 1, 0, 1))
+				.FillWidth(1)
+				[
+					SNew(STextBlock)
+					.Text(LOCTEXT("EncoderMappingValueLabel", "Encoder mapping :"))
+				]
+			]
+			+ SSplitter::Slot()
+			[
+				SNew(SNumericEntryBox<float>)
+				.AllowSpin(true)
+				.Value_Lambda([this]() { return EncoderMappingValue; })
+				.OnValueChanged(SNumericEntryBox<float>::FOnValueChanged::CreateLambda([this](float NewValue) { EncoderMappingValue = NewValue; }))
+				.OnValueCommitted(SNumericEntryBox<float>::FOnValueCommitted::CreateLambda([this](float NewValue, ETextCommit::Type CommitType) { EncoderMappingValue = NewValue; }))
+			]
+		];
+}
+
+void SLensDataPointDialog::SetDataCategory(EDataCategories NewCategory)
+{
+	if (NewCategory != SelectedCategory)
+	{
+		SelectedCategory = NewCategory;
+		EncoderMappingValue = 0.0f;
+	}
+
+	//Updates displayed widgets based on category
+	TrackingDataContainer->SetContent(MakeTrackingDataWidget());
+	LensDataContainer->SetContent(MakeLensDataWidget());
+}
+
+void SLensDataPointDialog::CloseDialog()
+{
+	TSharedPtr<SWindow> ExistingWindowPin = ExistingWindow.Pin();
+	if (ExistingWindowPin.IsValid())
+	{
+		ExistingWindowPin->RequestDestroyWindow();
+		ExistingWindowPin = nullptr;
+	}
+}
+
+bool SLensDataPointDialog::IsTrackingDataOverrided(int32 Index) const
+{
+	if (ensure(Index >= 0 && Index <= 1))
+	{
+		return TrackingInputData[Index].bIsOverrided;
+	}
+
+	return false;
+}
+void SLensDataPointDialog::OnOverrideTrackingData(ECheckBoxState NewState, int32 Index)
+{
+	if (ensure(Index >= 0 && Index <= 1))
+	{
+		TrackingInputData[Index].bIsOverrided = NewState == ECheckBoxState::Checked;
+	}
+}
+
+float SLensDataPointDialog::GetTrackingData(int32 Index) const
+{
+	if (ensure(Index >= 0 && Index <= 1))
+	{
+		return TrackingInputData[Index].Value;
+	}
+
+	return 0.0f;
+}
+
+void SLensDataPointDialog::SetTrackingData(float Value, int32 Index)
+{
+	if (ensure(Index >= 0 && Index <= 1))
+	{
+		TrackingInputData[Index].Value = Value;
+	}
+}
+
+void SLensDataPointDialog::RefreshEvaluationData()
+{
+	//Get FIZ coming from LiveLink. If normalized values are available, take them. Otherwise, look for precalibrated values or default to 0.0f
+	const FCachedFIZData& CachedData = CachedFIZ.Get();
+	const float Focus = CachedData.NormalizedFocus.IsSet() ? CachedData.NormalizedFocus.GetValue() : CachedData.Focus.IsSet() ? CachedData.Focus.GetValue() : 0.0f;
+	const float Iris = CachedData.NormalizedIris.IsSet() ? CachedData.NormalizedIris.GetValue() : CachedData.Iris.IsSet() ? CachedData.Iris.GetValue() : 0.0f;
+	const float Zoom = CachedData.NormalizedZoom.IsSet() ? CachedData.NormalizedZoom.GetValue() : CachedData.Zoom.IsSet() ? CachedData.Zoom.GetValue() : 0.0f;
+	
+	switch (SelectedCategory)
+	{
+		case EDataCategories::Focus:
+		{
+			if (IsTrackingDataOverrided(0) == false)
+			{
+				TrackingInputData[0].Value = Focus;
+			}
+			break;
+		}
+		case EDataCategories::Iris:
+		{
+			if (IsTrackingDataOverrided(0) == false)
+			{
+				TrackingInputData[0].Value = Iris;
+			}
+			break;
+		}
+		break;
+		case EDataCategories::Zoom:
+		{
+			if (IsTrackingDataOverrided(0) == false)
+			{
+				TrackingInputData[0].Value = Zoom;
+			}
+			break;
+		}
+		break;
+		case EDataCategories::Distortion:
+		case EDataCategories::ImageCenter:
+		case EDataCategories::NodalOffset:
+		case EDataCategories::STMap:
+		default:
+		{
+			if (IsTrackingDataOverrided(0) == false)
+			{
+				TrackingInputData[0].Value = Focus;
+			}
+			if (IsTrackingDataOverrided(1) == false)
+			{
+				TrackingInputData[1].Value = Zoom;
+			}
+			break;
+		}
+	}
+}
+
+void SLensDataPointDialog::AddDataToLensFile()
+{
+	//Logic to add points to LensFile is still todo. Requires updating LensFile API
+	switch (SelectedCategory)
+	{
+		case EDataCategories::Focus:
+		{
+			LensFile->EncoderMapping.Focus.AddKey(TrackingInputData[0].Value, EncoderMappingValue);
+			break;
+		}
+		case EDataCategories::Iris:
+		{
+			LensFile->EncoderMapping.Iris.AddKey(TrackingInputData[0].Value, EncoderMappingValue);
+			break;
+		}
+		break;
+		case EDataCategories::Zoom:
+		{
+			LensFile->EncoderMapping.Zoom.AddKey(TrackingInputData[0].Value, EncoderMappingValue);
+			break;
+		}
+		break;
+		case EDataCategories::Distortion:
+		{
+			break;
+		}	
+		case EDataCategories::ImageCenter:
+		{
+			break;
+		}
+		case EDataCategories::NodalOffset:
+		{
+			break;
+		}
+		case EDataCategories::STMap:
+		{
+			break;
+		}
+		default:
+			break;
+	}
+}
+
+#undef LOCTEXT_NAMESPACE 
diff --git a/Engine/Plugins/VirtualProduction/CameraCalibration/Source/CameraCalibrationEditor/Private/AssetEditor/SLensDataPointDialog.h b/Engine/Plugins/VirtualProduction/CameraCalibration/Source/CameraCalibrationEditor/Private/AssetEditor/SLensDataPointDialog.h
new file mode 100644
--- /dev/null
+++ b/Engine/Plugins/VirtualProduction/CameraCalibration/Source/CameraCalibrationEditor/Private/AssetEditor/SLensDataPointDialog.h
@@ -0,0 +1,119 @@
+// Copyright Epic Games, Inc. All Rights Reserved.
+
+#pragma once
+
+#include "Widgets/SCompoundWidget.h"
+
+#include "CameraCalibrationEditorCommon.h"
+#include "SLensEvaluation.h"
+#include "UObject/StrongObjectPtr.h"
+
+class SWindow;
+class ULensFile;
+
+class SLensDataPointDialog : public SCompoundWidget
+{
+private:
+	using Super = SCompoundWidget;
+
+public:
+
+	SLATE_BEGIN_ARGS(SLensDataPointDialog)
+	{}
+
+		/** FIZ data */
+		SLATE_ATTRIBUTE(FCachedFIZData, CachedFIZData)
+
+		/** If bound, used to notify a point was added to LensFile. */
+		SLATE_EVENT(FSimpleDelegate, OnDataPointAdded)
+
+	SLATE_END_ARGS()
+
+	void Construct(const FArguments& InArgs, ULensFile* InLensFile);
+
+	virtual void Tick(const FGeometry& AllottedGeometry, const double InCurrentTime, const float InDeltaTime) override;
+
+	/** Finds currently opened dialog or spawns a new one */
+	static void OpenDialog(ULensFile* InLensFile, TAttribute<FCachedFIZData> InCachedFIZData, const FSimpleDelegate& InOnDataPointAdded);
+
+private:
+
+	/** Called when AddPoint button is clicked */
+	FReply OnAddDataPointClicked();
+	
+	/** Called when Cancel button is clicked */
+	FReply OnCancelDataPointClicked();
+
+	/** Makes the different widgets used in this dialog */
+	TSharedRef<SWidget> MakeTrackingDataWidget();
+	TSharedRef<SWidget> MakeLensDataWidget();
+	TSharedRef<SWidget> MakeButtonsWidget();
+	TSharedRef<SWidget> MakeDataCategoryMenuWidget();
+	TSharedRef<SWidget> MakeEncoderMappingWidget();
+
+	/** Updates the desired data category */
+	void SetDataCategory(EDataCategories NewCategory);
+
+	/** Closes this dialog */
+	void CloseDialog();
+
+	/** Whether tracking data is overrided or not */
+	bool IsTrackingDataOverrided(int32 Index) const;
+	
+	/** Triggered when override TrackingData state changes */
+	void OnOverrideTrackingData(ECheckBoxState NewState, int32 Index);
+
+	/** Gets tracking data for Index */
+	float GetTrackingData(int32 Index) const;
+
+	/** Updates tracking data for desired Index */
+	void SetTrackingData(float Value, int32 Index);
+
+	/** Updates the current inputs values */
+	void RefreshEvaluationData();
+	
+	/** Called when ready to add data to LensFile */
+	void AddDataToLensFile();
+
+private:
+
+	/**
+	 * Used to build tracking input widget
+	 * Depending on the category, there is either one
+	 * input (Encoder) or two (Focus / Zoom)
+	 */
+	struct FTrackingInputData
+	{
+		FString Label;
+		bool bIsOverrided = false;
+		float Value = 0.0f;
+	};
+	FTrackingInputData TrackingInputData[2];
+
+	/** LensFile being edited */
+	TStrongObjectPtr<ULensFile> LensFile;
+
+	/** Singleton for the pop-up window. */
+	static TWeakPtr<SWindow> ExistingWindow;
+
+	/** Container for tracking inputs */
+	TSharedPtr<SBorder> TrackingDataContainer;
+
+	/** Container for lens data */
+	TSharedPtr<SBorder> LensDataContainer;
+	
+	/** Category of data being added */
+	EDataCategories SelectedCategory = EDataCategories::NodalOffset;
+
+	/** When adding an encoder mapping value, float holding the current one */
+	float EncoderMappingValue = 0.0f;
+
+	/** LensData currently being added */
+	TSharedPtr<FStructOnScope> LensData;
+
+	/** Evaluated FIZ for the current frame */
+	TAttribute<FCachedFIZData> CachedFIZ;
+
+	/** Delegate to call when point is added to LensFile */
+	FSimpleDelegate OnDataPointAdded;
+};
\ No newline at end of file
diff --git a/Engine/Plugins/VirtualProduction/CameraCalibration/Source/CameraCalibrationEditor/Private/AssetEditor/SLensDataViewer.cpp b/Engine/Plugins/VirtualProduction/CameraCalibration/Source/CameraCalibrationEditor/Private/AssetEditor/SLensDataViewer.cpp
new file mode 100644
--- /dev/null
+++ b/Engine/Plugins/VirtualProduction/CameraCalibration/Source/CameraCalibrationEditor/Private/AssetEditor/SLensDataViewer.cpp
@@ -0,0 +1,531 @@
+// Copyright Epic Games, Inc. All Rights Reserved.
+
+#include "SLensDataViewer.h"
+
+#include "CurveEditor.h"
+#include "ICurveEditorModule.h"
+#include "ISinglePropertyView.h"
+#include "LensFile.h"
+#include "RichCurveEditorModel.h"
+#include "SCurveEditorPanel.h"
+#include "SLensDataPointDialog.h"
+#include "Widgets/Input/SButton.h"
+
+
+#define LOCTEXT_NAMESPACE "LensDataViewer"
+
+
+namespace LensDataUtils
+{
+	const FName EncoderCategoryLabel(TEXT("Encoders"));
+	const FName EncoderFocusLabel(TEXT("Focus"));
+	const FName EncoderIrisLabel(TEXT("Iris"));
+	const FName EncoderZoomLabel(TEXT("Zoom"));
+	const FName DistortionCategoryLabel(TEXT("Distortion"));
+	const FName FxLabel(TEXT("Fx"));
+	const FName FyLabel(TEXT("Fy"));
+	const FName MapsCategoryLabel(TEXT("STMaps"));
+	const FName ImageCenterCategory(TEXT("Image Center"));
+	const FName CxLabel(TEXT("Cx"));
+	const FName CyLabel(TEXT("Cy"));
+	const FName NodalOffsetCategoryLabel(TEXT("Nodal Offset"));
+	const FName LocationXLabel(TEXT("Location - X"));
+	const FName LocationYLabel(TEXT("Location - Y"));
+	const FName LocationZLabel(TEXT("Location - Z"));
+	const FName RotationXLabel(TEXT("Rotation - X"));
+	const FName RotationYLabel(TEXT("Rotation - Y"));
+	const FName RotationZLabel(TEXT("Rotation - Z"));
+
+	template<typename Type>
+	void FindFocusPoints(const TArray<Type>&SourceData, TArray<TSharedPtr<FLensDataItem>>& OutDataItems)
+	{
+		TSet<float> FocusPoints;
+		FocusPoints.Reserve(SourceData.Num());
+		OutDataItems.Reserve(SourceData.Num());
+		for (const Type& Point : SourceData)
+		{
+			if (FocusPoints.Contains(Point.Focus) == false)
+			{
+				FocusPoints.Emplace(Point.Focus);
+				OutDataItems.Add(MakeShared<FLensDataItem>(Point.Focus));
+			}
+		}
+	}
+}
+
+
+FLensDataItem::FLensDataItem(float InFocus)
+	: Focus(InFocus)
+{
+
+}
+
+TSharedRef<ITableRow> FLensDataItem::MakeTreeRowWidget(const TSharedRef<STableViewBase>& InOwnerTable)
+{
+	return SNew(SLensDataItem, InOwnerTable, AsShared());
+}
+
+void SLensDataItem::Construct(const FArguments& InArgs, const TSharedRef<STableViewBase>& OwnerTable, TSharedRef<FLensDataItem> InItemData)
+{
+	WeakItem = InItemData;
+
+	STableRow<TSharedPtr<FLensDataItem>>::Construct(
+		STableRow<TSharedPtr<FLensDataItem>>::FArguments()
+		.Content()
+		[
+			SNew(SHorizontalBox)
+			+ SHorizontalBox::Slot()
+			.Padding(5.0f, 5.0f)
+			.AutoWidth()
+			[
+				SNew(STextBlock)
+				.Text(LOCTEXT("FocusLabel", "Focus: "))
+			]
+			+ SHorizontalBox::Slot()
+			.Padding(5.0f, 5.0f)
+			[
+				SNew(STextBlock)
+				.Text(this, &SLensDataItem::GetLabelText)
+			]
+		], OwnerTable);
+}
+
+FText SLensDataItem::GetLabelText() const
+{
+	return (FText::AsNumber(WeakItem.Pin()->Focus));
+}
+
+FLensDataCategoryItem::FLensDataCategoryItem(TWeakPtr<FLensDataCategoryItem> InParent, EDataCategories InCategory, FName InLabel)
+	: Category(InCategory)
+	, Label(InLabel)
+	, Parent(InParent)
+{
+
+}
+
+TSharedRef<ITableRow> FLensDataCategoryItem::MakeTreeRowWidget(const TSharedRef<STableViewBase>& InOwnerTable)
+{
+	return SNew(SLensDataCategoryItem, InOwnerTable, AsShared());
+}
+
+void SLensDataCategoryItem::Construct(const FArguments& InArgs, const TSharedRef<STableViewBase>& OwnerTable, TSharedRef<FLensDataCategoryItem> InItemData)
+{
+	WeakItem = InItemData;
+
+	STableRow<TSharedPtr<FLensDataCategoryItem>>::Construct(
+		STableRow<TSharedPtr<FLensDataCategoryItem>>::FArguments()
+		.Content()
+		[
+			SNew(STextBlock)
+			.Text(this, &SLensDataCategoryItem::GetLabelText)
+		], OwnerTable);
+}
+
+FText SLensDataCategoryItem::GetLabelText() const
+{
+	return (FText::FromName(WeakItem.Pin()->Label));
+}
+
+
+void SLensDataViewer::Construct(const FArguments& InArgs, ULensFile* InLensFile)
+{
+	LensFile = TStrongObjectPtr<ULensFile>(InLensFile);
+
+	//Setup curve editor
+	CurveEditor = MakeShared<FCurveEditor>();
+	FCurveEditorInitParams InitParams;
+	CurveEditor->InitCurveEditor(InitParams);
+	CurveEditor->GridLineLabelFormatXAttribute = LOCTEXT("GridXLabelFormat", "{0}");
+
+	TUniquePtr<ICurveEditorBounds> EditorBounds = MakeUnique<FStaticCurveEditorBounds>();
+	EditorBounds->SetInputBounds(0.05, 1.05);
+	CurveEditor->SetBounds(MoveTemp(EditorBounds));
+
+	CurvePanel = SNew(SCurveEditorPanel, CurveEditor.ToSharedRef());
+	CurveEditor->ZoomToFit();
+
+	CachedFIZ = InArgs._CachedFIZData;
+
+	ChildSlot
+	[
+		SNew(SHorizontalBox)
+		+ SHorizontalBox::Slot()
+		.FillWidth(0.4f)
+		[
+			MakeLensDataWidget()
+		]
+		+ SHorizontalBox::Slot()
+		.FillWidth(0.6f)
+		[
+			SNew(SVerticalBox)
+			+ SVerticalBox::Slot()
+			.AutoHeight()
+			[
+				MakeToolbarWidget()
+			]
+			+ SVerticalBox::Slot()
+			[
+				CurvePanel.ToSharedRef()
+			]
+		]
+	];
+
+	Refresh();
+}
+
+TSharedPtr<FLensDataCategoryItem> SLensDataViewer::GetDataCategorySelection() const
+{
+	TArray<TSharedPtr<FLensDataCategoryItem>> SelectedNodes;
+	TreeView->GetSelectedItems(SelectedNodes);
+	if (SelectedNodes.Num())
+	{
+		return SelectedNodes[0];
+	}
+	return nullptr;
+}
+
+TSharedRef<ITableRow> SLensDataViewer::OnGenerateDataCategoryRow(TSharedPtr<FLensDataCategoryItem> Item, const TSharedRef<STableViewBase>& OwnerTable)
+{
+	return Item->MakeTreeRowWidget(OwnerTable);
+}
+
+void SLensDataViewer::OnGetDataCategoryItemChildren(TSharedPtr<FLensDataCategoryItem> Item, TArray<TSharedPtr<FLensDataCategoryItem>>& OutChildren)
+{
+	if (Item.IsValid())
+	{
+		OutChildren = Item->Children;
+	}
+}
+
+void SLensDataViewer::OnDataCategorySelectionChanged(TSharedPtr<FLensDataCategoryItem> Item, ESelectInfo::Type SelectInfo)
+{
+	if (!Item || SelectInfo != ESelectInfo::OnMouseClick)
+	{
+		return;
+	}
+
+	RefreshDataEntriesTree();
+}
+
+TSharedPtr<FLensDataItem> SLensDataViewer::GetSelectedDataEntry() const
+{
+	TArray<TSharedPtr<FLensDataItem>> SelectedNodes;
+	DataEntriesTree->GetSelectedItems(SelectedNodes);
+	if (SelectedNodes.Num())
+	{
+		return SelectedNodes[0];
+	}
+	return nullptr;
+}
+
+TSharedRef<ITableRow> SLensDataViewer::OnGenerateDataEntryRow(TSharedPtr<FLensDataItem> Node, const TSharedRef<STableViewBase>& OwnerTable)
+{
+	return Node->MakeTreeRowWidget(OwnerTable);
+}
+
+void SLensDataViewer::OnGetDataEntryChildren(TSharedPtr<FLensDataItem> Node, TArray<TSharedPtr<FLensDataItem>>& OutNodes)
+{
+	if (Node.IsValid())
+	{
+		OutNodes = Node->Children;
+	}
+}
+
+void SLensDataViewer::OnDataEntrySelectionChanged(TSharedPtr<FLensDataItem> Node, ESelectInfo::Type SelectInfo)
+{
+	RefreshCurve();
+}
+
+TSharedRef<SWidget> SLensDataViewer::MakeLensDataWidget()
+{
+	const FSinglePropertyParams InitParams;
+	FPropertyEditorModule& PropertyEditorModule = FModuleManager::GetModuleChecked<FPropertyEditorModule>("PropertyEditor");
+	const TSharedPtr<ISinglePropertyView> DataModeWidget = PropertyEditorModule.CreateSingleProperty(
+		LensFile.Get()
+		, GET_MEMBER_NAME_CHECKED(ULensFile, DataMode)
+		, InitParams);
+
+	FSimpleDelegate OnDataModeChangedDelegate = FSimpleDelegate::CreateSP(this, &SLensDataViewer::OnDataModeChanged);
+	DataModeWidget->SetOnPropertyValueChanged(OnDataModeChangedDelegate);
+
+	return 
+		SNew(SVerticalBox)
+		+ SVerticalBox::Slot()
+		.Padding(5.0f, 5.0f)
+		.AutoHeight()
+		[
+			SNew(SBorder)
+			.HAlign(HAlign_Fill)
+			.VAlign(VAlign_Fill)
+			.Padding(FMargin(4.f, 4.f, 4.f, 4.f))
+			[
+				DataModeWidget.ToSharedRef()
+			]
+		]
+		+ SVerticalBox::Slot()
+		.Padding(5.0f, 5.0f)
+		.FillHeight(.5f)
+		[
+			SNew(SBorder)
+			.HAlign(HAlign_Fill)
+			.VAlign(VAlign_Fill)
+			.Padding(FMargin(4.f, 4.f, 4.f, 4.f))
+			[
+				SAssignNew(TreeView, STreeView<TSharedPtr<FLensDataCategoryItem>>)
+				.TreeItemsSource(&DataCategories)
+				.ItemHeight(24.0f)
+				.OnGenerateRow(this, &SLensDataViewer::OnGenerateDataCategoryRow)
+				.OnGetChildren(this, &SLensDataViewer::OnGetDataCategoryItemChildren)
+				.OnSelectionChanged(this, &SLensDataViewer::OnDataCategorySelectionChanged)
+				.ClearSelectionOnClick(false)
+			]
+			
+		]
+		+ SVerticalBox::Slot()
+		.Padding(5.0f, 5.0f)
+		.FillHeight(.5f)
+		[
+			SNew(SBorder)
+			.HAlign(HAlign_Fill)
+			.VAlign(VAlign_Fill)
+			.Padding(FMargin(4.f, 4.f, 4.f, 4.f))
+			[
+				SNew(SVerticalBox)
+				+ SVerticalBox::Slot()
+				.AutoHeight()
+				[
+					SAssignNew(DataEntryNameWidget, STextBlock)
+				]
+				+ SVerticalBox::Slot()
+				[
+					SAssignNew(DataEntriesTree, STreeView<TSharedPtr<FLensDataItem>>)
+					.TreeItemsSource(&DataEntries)
+					.ItemHeight(24.0f)
+					.OnGenerateRow(this, &SLensDataViewer::OnGenerateDataEntryRow)
+					.OnGetChildren(this, &SLensDataViewer::OnGetDataEntryChildren)
+					.OnSelectionChanged(this, &SLensDataViewer::OnDataEntrySelectionChanged)
+					.ClearSelectionOnClick(false)
+				]
+			]
+		];
+}
+
+
+TSharedRef<SWidget> SLensDataViewer::MakeToolbarWidget()
+{
+	TSharedRef<SWidget> Toolbar = SNew(SBorder)
+	.Padding(FMargin(5.f, 5.f, 5.f, 5.f))
+	[
+		SNew(SHorizontalBox)
+		+ SHorizontalBox::Slot()
+		.AutoWidth()
+		.Padding(4.f)
+		[
+			SNew(SButton)
+			.ButtonStyle(FEditorStyle::Get(), "ToggleButton")
+			.ContentPadding(FMargin(4.f, 4.f))
+			.OnClicked(this, &SLensDataViewer::OnAddDataPointClicked)
+			.IsEnabled(MakeAttributeLambda([this] { return true; }))
+			[
+				SNew(STextBlock)
+				.Font(FEditorStyle::Get().GetFontStyle("FontAwesome.14"))
+				.Text(FText::FromString(FString(TEXT("\xf067"))) /*fa-plus*/)
+				.ColorAndOpacity(FLinearColor::White)
+			]
+		]
+	];
+
+	return Toolbar;
+}
+
+FReply SLensDataViewer::OnAddDataPointClicked()
+{
+	const FSimpleDelegate OnDataPointAdded = FSimpleDelegate::CreateSP(this, &SLensDataViewer::OnLensDataPointAdded);
+	SLensDataPointDialog::OpenDialog(LensFile.Get(), CachedFIZ, OnDataPointAdded);
+	return FReply::Handled();
+}
+
+void SLensDataViewer::OnDataModeChanged()
+{
+	Refresh();
+}
+
+void SLensDataViewer::Refresh()
+{
+	RefreshDataCategoriesTree();
+	RefreshDataEntriesTree();
+}
+
+void SLensDataViewer::RefreshDataCategoriesTree()
+{
+	//Builds the data category tree
+	
+	DataCategories.Reset();
+
+	DataCategories.Add(MakeShared<FLensDataCategoryItem>(nullptr, EDataCategories::Focus,  LensDataUtils::EncoderFocusLabel));
+	DataCategories.Add(MakeShared<FLensDataCategoryItem>(nullptr, EDataCategories::Iris, LensDataUtils::EncoderIrisLabel));
+	DataCategories.Add(MakeShared<FLensDataCategoryItem>(nullptr, EDataCategories::Zoom, LensDataUtils::EncoderZoomLabel));
+	
+	if (LensFile->DataMode == ELensDataMode::Parameters)
+	{
+		TSharedPtr<FLensDataCategoryItem> DistortionEntry = MakeShared<FLensDataCategoryItem>(nullptr, EDataCategories::Distortion, LensDataUtils::DistortionCategoryLabel);
+		DataCategories.Add(DistortionEntry);
+
+		TArray<FText> Parameters;
+		if (LensFile->LensInfo.LensModel)
+		{
+			Parameters = LensFile->LensInfo.LensModel.GetDefaultObject()->GetParameterDisplayNames();
+		}
+
+		for (const FText& Parameter : Parameters)
+		{
+			DistortionEntry->Children.Add(MakeShared<FLensDataCategoryItem>(DistortionEntry, EDataCategories::Distortion, *Parameter.ToString()));
+		}
+
+		DistortionEntry->Children.Add(MakeShared<FLensDataCategoryItem>(DistortionEntry, EDataCategories::Distortion, LensDataUtils::FxLabel));
+		DistortionEntry->Children.Add(MakeShared<FLensDataCategoryItem>(DistortionEntry, EDataCategories::Distortion, LensDataUtils::FyLabel));
+	}
+	else
+	{
+		DataCategories.Add(MakeShared<FLensDataCategoryItem>(nullptr, EDataCategories::STMap, LensDataUtils::MapsCategoryLabel));
+	}
+
+	
+	TSharedPtr<FLensDataCategoryItem> ImageCenterEntry = MakeShared<FLensDataCategoryItem>(nullptr, EDataCategories::ImageCenter, LensDataUtils::ImageCenterCategory);
+	DataCategories.Add(ImageCenterEntry);
+	ImageCenterEntry->Children.Add(MakeShared<FLensDataCategoryItem>(ImageCenterEntry, EDataCategories::ImageCenter, LensDataUtils::CxLabel));
+	ImageCenterEntry->Children.Add(MakeShared<FLensDataCategoryItem>(ImageCenterEntry, EDataCategories::ImageCenter, LensDataUtils::CyLabel));
+
+	TSharedPtr<FLensDataCategoryItem> NodalOffsetCategory = MakeShared<FLensDataCategoryItem>(nullptr, EDataCategories::NodalOffset, LensDataUtils::NodalOffsetCategoryLabel);
+	DataCategories.Add(NodalOffsetCategory);
+
+	NodalOffsetCategory->Children.Add(MakeShared<FLensDataCategoryItem>(NodalOffsetCategory, EDataCategories::NodalOffset, LensDataUtils::LocationXLabel));
+	NodalOffsetCategory->Children.Add(MakeShared<FLensDataCategoryItem>(NodalOffsetCategory, EDataCategories::NodalOffset, LensDataUtils::LocationYLabel));
+	NodalOffsetCategory->Children.Add(MakeShared<FLensDataCategoryItem>(NodalOffsetCategory, EDataCategories::NodalOffset, LensDataUtils::LocationZLabel));
+	NodalOffsetCategory->Children.Add(MakeShared<FLensDataCategoryItem>(NodalOffsetCategory, EDataCategories::NodalOffset, LensDataUtils::RotationXLabel));
+	NodalOffsetCategory->Children.Add(MakeShared<FLensDataCategoryItem>(NodalOffsetCategory, EDataCategories::NodalOffset, LensDataUtils::RotationYLabel));
+	NodalOffsetCategory->Children.Add(MakeShared<FLensDataCategoryItem>(NodalOffsetCategory, EDataCategories::NodalOffset, LensDataUtils::RotationZLabel));
+
+	TreeView->RequestTreeRefresh();
+}
+
+void SLensDataViewer::RefreshDataEntriesTree()
+{
+	DataEntries.Reset();
+
+	if (TSharedPtr<FLensDataCategoryItem> CategoryItem = GetDataCategorySelection())
+	{
+		DataEntryNameWidget->SetText(FText::FromName(CategoryItem->Label));
+
+		switch (CategoryItem->Category)
+		{
+			case EDataCategories::Distortion:
+			{
+				LensDataUtils::FindFocusPoints<FDistortionMapPoint>(LensFile->DistortionMapping, DataEntries);
+				break;
+			}
+			case EDataCategories::ImageCenter:
+			{
+				LensDataUtils::FindFocusPoints<FIntrinsicMapPoint>(LensFile->IntrinsicMapping, DataEntries);
+				break;
+			}
+			case EDataCategories::NodalOffset:
+			{
+				LensDataUtils::FindFocusPoints<FNodalOffsetMapPoint>(LensFile->NodalOffsetMapping, DataEntries);
+				break;
+			}
+			case EDataCategories::STMap:
+			{
+				LensDataUtils::FindFocusPoints<FCalibratedMapPoint>(LensFile->CalibratedMapPoints, DataEntries);
+				break;
+			}
+			default:
+			{
+				break;
+			}
+		}
+	}
+	else
+	{
+		DataEntryNameWidget->SetText(LOCTEXT("NoCategorySelected", "Select a category"));
+	}
+
+	//When data entries have been repopulated, refresh the tree and select first item
+	DataEntriesTree->RequestListRefresh();
+
+	if (DataEntries.Num())
+	{
+		DataEntriesTree->SetSelection(DataEntries[0]);
+	}
+	else
+	{
+		DataEntriesTree->SetSelection(nullptr);
+	}
+}
+
+void SLensDataViewer::RefreshCurve()
+{
+	CurveEditor->RemoveAllCurves();
+	TUniquePtr<FRichCurveEditorModelRaw> NewCurve;
+
+	TSharedPtr<FLensDataCategoryItem> CategoryItem = GetDataCategorySelection();
+	if (CategoryItem.IsValid())
+	{
+		switch (CategoryItem->Category)
+		{
+			case EDataCategories::Focus:
+			{
+				NewCurve = MakeUnique<FRichCurveEditorModelRaw>(&LensFile->EncoderMapping.Focus, LensFile.Get());
+				break;
+			}
+			case EDataCategories::Iris:
+			{
+				NewCurve = MakeUnique<FRichCurveEditorModelRaw>(&LensFile->EncoderMapping.Iris, LensFile.Get());
+				break;
+			}
+			case EDataCategories::Zoom:
+			{
+				NewCurve = MakeUnique<FRichCurveEditorModelRaw>(&LensFile->EncoderMapping.Zoom, LensFile.Get());
+				break;
+			}
+			case EDataCategories::Distortion:
+			{
+				break;
+			}
+			case EDataCategories::ImageCenter:
+			{
+				break;
+			}
+			case EDataCategories::NodalOffset:
+			{
+				break;
+			}
+			case EDataCategories::STMap:
+			{
+				break;
+			}
+		default:
+			{
+				break;
+			}
+		}
+	}
+
+	//If a curve was setup, add it to the editor
+	if (NewCurve)
+	{
+		NewCurve->SetShortDisplayName(FText::FromName(CategoryItem->Label));
+		NewCurve->SetColor(FLinearColor(0.0f, 1.0f, 0.0f, 1.0f));
+		const FCurveModelID CurveId = CurveEditor->AddCurve(MoveTemp(NewCurve));
+		CurveEditor->PinCurve(CurveId);
+	}
+}
+
+void SLensDataViewer::OnLensDataPointAdded()
+{
+	RefreshDataEntriesTree();
+}
+
+#undef LOCTEXT_NAMESPACE /* LensDataViewer */
+
+
diff --git a/Engine/Plugins/VirtualProduction/CameraCalibration/Source/CameraCalibrationEditor/Private/AssetEditor/SLensDataViewer.h b/Engine/Plugins/VirtualProduction/CameraCalibration/Source/CameraCalibrationEditor/Private/AssetEditor/SLensDataViewer.h
new file mode 100644
--- /dev/null
+++ b/Engine/Plugins/VirtualProduction/CameraCalibration/Source/CameraCalibrationEditor/Private/AssetEditor/SLensDataViewer.h
@@ -0,0 +1,195 @@
+// Copyright Epic Games, Inc. All Rights Reserved.
+
+#pragma once
+
+#include "Widgets/SCompoundWidget.h"
+
+#include "CameraCalibrationEditorCommon.h"
+#include "SLensFilePanel.h"
+#include "Widgets/DeclarativeSyntaxSupport.h"
+#include "Widgets/Views/STreeView.h"
+#include "UObject/StrongObjectPtr.h"
+
+class FCurveEditor;
+class SCurveEditorPanel;
+class ULensFile;
+
+/**
+ * Data entry item
+ * @note To be expanded
+ */
+class FLensDataItem : public TSharedFromThis<FLensDataItem>
+{
+public:
+	FLensDataItem(float InFocus);
+
+	TSharedRef<ITableRow> MakeTreeRowWidget(const TSharedRef<STableViewBase>& InOwnerTable);
+
+	/** Focus value of this item */
+	float Focus;
+
+	/** Children of this item */
+	TArray<TSharedPtr<FLensDataItem>> Children;
+};
+
+/**
+ * Widget representing a data entry row
+ * @note To be expanded
+ */
+class SLensDataItem : public STableRow<TSharedPtr<FLensDataItem>>
+{
+	SLATE_BEGIN_ARGS(SLensDataItem) {}
+	SLATE_END_ARGS()
+
+	void Construct(const FArguments& InArgs, const TSharedRef<STableViewBase>& OwnerTable, TSharedRef<FLensDataItem> InItemData);
+
+private:
+	FText GetLabelText() const;
+
+private:
+
+	/** WeakPtr to source data item */
+	TWeakPtr<FLensDataItem> WeakItem;
+};
+
+/**
+ * Data category item
+ */
+class FLensDataCategoryItem : public TSharedFromThis<FLensDataCategoryItem>
+{
+public:
+	FLensDataCategoryItem(TWeakPtr<FLensDataCategoryItem> Parent, EDataCategories InCategory, FName InLabel);
+
+	TSharedRef<ITableRow> MakeTreeRowWidget(const TSharedRef<STableViewBase>& InOwnerTable);
+
+public:
+
+	/** Category this item is associated with */
+	EDataCategories Category;
+
+	/** Label of this category */
+	FName Label;
+
+	/** WeakPtr to parent of this item */
+	TWeakPtr<FLensDataCategoryItem> Parent;
+
+	/** Children of this category */
+	TArray<TSharedPtr<FLensDataCategoryItem>> Children;
+};
+
+/**
+ * Data category row widget
+ */
+class SLensDataCategoryItem : public STableRow<TSharedPtr<FLensDataCategoryItem>>
+{
+	SLATE_BEGIN_ARGS(SLensDataCategoryItem) {}
+	SLATE_END_ARGS()
+
+	void Construct(const FArguments& InArgs, const TSharedRef<STableViewBase>& OwnerTable, TSharedRef<FLensDataCategoryItem> InItemData);
+
+private:
+
+	/** Returns the label of this row */
+	FText GetLabelText() const;
+
+private:
+
+	/** WeakPtr to source data item */
+	TWeakPtr<FLensDataCategoryItem> WeakItem;
+};
+
+/** Widget used to display data from the LensFile */
+class SLensDataViewer : public SCompoundWidget
+{
+public:
+
+	SLATE_BEGIN_ARGS(SLensDataViewer)
+	{}
+
+		/** FIZ data */
+		SLATE_ATTRIBUTE(FCachedFIZData, CachedFIZData)
+
+	SLATE_END_ARGS()
+
+	void Construct(const FArguments& InArgs, ULensFile* InLensFile);
+
+	/** Get the currently selected data category. */
+	TSharedPtr<FLensDataCategoryItem> GetDataCategorySelection() const;
+
+	/** Generates one data category row */
+	TSharedRef<ITableRow> OnGenerateDataCategoryRow(TSharedPtr<FLensDataCategoryItem> Item, const TSharedRef<STableViewBase>& OwnerTable);
+
+	/** Used to get the children of a data category item */
+	void OnGetDataCategoryItemChildren(TSharedPtr<FLensDataCategoryItem> Item, TArray<TSharedPtr<FLensDataCategoryItem>>& OutChildren);
+
+	/** Triggered when data category selection has changed */
+	void OnDataCategorySelectionChanged(TSharedPtr<FLensDataCategoryItem> Item, ESelectInfo::Type SelectInfo);
+
+	/** Get the currently selected data entry */
+	TSharedPtr<FLensDataItem> GetSelectedDataEntry() const;
+
+	/** Generates one data entry row */
+	TSharedRef<ITableRow> OnGenerateDataEntryRow(TSharedPtr<FLensDataItem> Node, const TSharedRef<STableViewBase>& OwnerTable);
+
+	/** Used to get the children of a data entry item */
+	void OnGetDataEntryChildren(TSharedPtr<FLensDataItem> Node, TArray<TSharedPtr<FLensDataItem>>& OutNodes);
+
+	/** Triggered when data entry selection has changed */
+	void OnDataEntrySelectionChanged(TSharedPtr<FLensDataItem> Node, ESelectInfo::Type SelectInfo);
+
+private:
+
+	/** Makes the widget showing lens data and curve editor */
+	TSharedRef<SWidget> MakeLensDataWidget();
+
+	/** Makes the Toolbar with data manipulation buttons */
+	TSharedRef<SWidget> MakeToolbarWidget();
+
+	/** Called when user clicks on AddPoint button */
+	FReply OnAddDataPointClicked();
+
+	/** Called when DataMode is changed */
+	void OnDataModeChanged();
+
+	/** Refreshes the widget's trees */
+	void Refresh();
+
+	/** Refreshes data categories tree */
+	void RefreshDataCategoriesTree();
+
+	/** Refreshes data entries tree */
+	void RefreshDataEntriesTree();
+
+	/** Refreshes curve editor */
+	void RefreshCurve();
+
+	/** Callbacked when user clicks AddPoint from the dialog */
+	void OnLensDataPointAdded();
+
+private:
+	
+	/** Data category TreeView */
+	TSharedPtr<STreeView<TSharedPtr<FLensDataCategoryItem>>> TreeView;
+
+	/** List of data category items */
+	TArray<TSharedPtr<FLensDataCategoryItem>> DataCategories;
+
+	/** Data items associated with selected data category TreeView */
+	TSharedPtr<STreeView<TSharedPtr<FLensDataItem>>> DataEntriesTree;
+
+	/** List of data items for the selected data category */
+	TArray<TSharedPtr<FLensDataItem>> DataEntries;
+
+	/** Data entries title */
+	TSharedPtr<STextBlock> DataEntryNameWidget;
+
+	/** LensFile being edited */
+	TStrongObjectPtr<ULensFile> LensFile;
+
+	/** Curve editor manager and panel to display */
+	TSharedPtr<FCurveEditor> CurveEditor;
+	TSharedPtr<SCurveEditorPanel> CurvePanel;
+
+	/** Evaluated FIZ for the current frame */
+	TAttribute<FCachedFIZData> CachedFIZ;
+};
\ No newline at end of file
diff --git a/Engine/Plugins/VirtualProduction/CameraCalibration/Source/CameraCalibrationEditor/Private/AssetEditor/SLensEvaluation.cpp b/Engine/Plugins/VirtualProduction/CameraCalibration/Source/CameraCalibrationEditor/Private/AssetEditor/SLensEvaluation.cpp
new file mode 100644
--- /dev/null
+++ b/Engine/Plugins/VirtualProduction/CameraCalibration/Source/CameraCalibrationEditor/Private/AssetEditor/SLensEvaluation.cpp
@@ -0,0 +1,625 @@
+// Copyright Epic Games, Inc. All Rights Reserved.
+
+#include "SLensEvaluation.h"
+
+#include "Features/IModularFeatures.h"
+#include "ILiveLinkClient.h"
+#include "LensFile.h"
+#include "Roles/LiveLinkCameraRole.h"
+#include "Roles/LiveLinkCameraTypes.h"
+#include "UI/CameraCalibrationEditorStyle.h"
+#include "Widgets/Input/SCheckBox.h"
+#include "Widgets/Layout/SGridPanel.h"
+#include "Widgets/SBoxPanel.h"
+#include "Widgets/Text/STextBlock.h"
+
+
+#define LOCTEXT_NAMESPACE "LensEvaluation"
+
+void SLensEvaluation::Construct(const FArguments& InArgs, ULensFile* InLensFile)
+{
+	LensFile = TStrongObjectPtr<ULensFile>(InLensFile);
+
+	ChildSlot
+	[
+		SNew(SBorder)
+		.Padding(FMargin(4.f, 4.f, 4.f, 4.f))
+		[
+			SNew(SHorizontalBox)
+
+			//Tracking section
+			+ SHorizontalBox::Slot()
+			.Padding(5.0f, 5.0f)
+			.FillWidth(0.2f)
+			[
+				MakeTrackingWidget()
+			]
+			//FIZ section
+			+ SHorizontalBox::Slot()
+			.Padding(5.0f, 5.0f)
+			.FillWidth(0.2f)
+			[
+				MakeFIZWidget()
+			]
+			//Distortion section
+			+ SHorizontalBox::Slot()
+			.Padding(5.0f, 5.0f)
+			.FillWidth(0.2f)
+			[
+				MakeDistortionWidget()
+			]
+			//Image Center section
+			+ SHorizontalBox::Slot()
+			.Padding(5.0f, 5.0f)
+			.FillWidth(0.2f)
+			[
+				MakeIntrinsicsWidget()
+			]
+			//Nodal offset section
+			+ SHorizontalBox::Slot()
+			.Padding(5.0f, 5.0f)
+			.FillWidth(0.2f)
+			[
+				MakeNodalOffsetWidget()	
+			]
+		]
+	];
+}
+
+void SLensEvaluation::Tick(const FGeometry& AllottedGeometry, const double InCurrentTime, const float InDeltaTime)
+{
+	//Cache LiveLink data every tick to be sure we have the right one for the frame during calibration
+	CacheLiveLinkData();
+	CacheLensFileData();
+
+	Super::Tick(AllottedGeometry, InCurrentTime, InDeltaTime);
+}
+
+ECheckBoxState SLensEvaluation::IsTrackingActive() const
+{
+	return bIsUsingLiveLinkTracking ? ECheckBoxState::Checked : ECheckBoxState::Unchecked;
+}
+
+void SLensEvaluation::OnTrackingStateChanged(ECheckBoxState NewState)
+{
+	bIsUsingLiveLinkTracking = NewState == ECheckBoxState::Checked ? true : false;
+}
+
+bool SLensEvaluation::CanSelectTrackingSource() const
+{
+	return IsTrackingActive() == ECheckBoxState::Checked;
+}
+
+SLiveLinkSubjectRepresentationPicker::FLiveLinkSourceSubjectRole SLensEvaluation::GetTrackingSubject() const
+{
+	return SLiveLinkSubjectRepresentationPicker::FLiveLinkSourceSubjectRole(TrackingSource);
+}
+
+void SLensEvaluation::SetTrackingSubject(SLiveLinkSubjectRepresentationPicker::FLiveLinkSourceSubjectRole NewValue)
+{
+	TrackingSource = NewValue.ToSubjectRepresentation();
+}
+
+void SLensEvaluation::CacheLiveLinkData()
+{
+	if (ShouldUpdateTracking() == false)
+	{
+		return;
+	}
+
+	//Start clean
+	CachedLiveLinkData.NormalizedFocus.Reset();
+	CachedLiveLinkData.NormalizedIris.Reset();
+	CachedLiveLinkData.NormalizedZoom.Reset();
+	CachedLiveLinkData.Focus.Reset();
+	CachedLiveLinkData.Iris.Reset();
+	CachedLiveLinkData.Zoom.Reset();
+
+	if (TrackingSource.Role && TrackingSource.Role->IsChildOf(ULiveLinkCameraRole::StaticClass()))
+	{
+		IModularFeatures& ModularFeatures = IModularFeatures::Get();
+		if (ModularFeatures.IsModularFeatureAvailable(ILiveLinkClient::ModularFeatureName))
+		{
+			ILiveLinkClient* LiveLinkClient = &IModularFeatures::Get().GetModularFeature<ILiveLinkClient>(ILiveLinkClient::ModularFeatureName);
+			if(LiveLinkClient)
+			{
+				FLiveLinkSubjectFrameData SubjectData;
+				if (LiveLinkClient->EvaluateFrame_AnyThread(TrackingSource.Subject, TrackingSource.Role, SubjectData))
+				{
+					FLiveLinkCameraStaticData* StaticData = SubjectData.StaticData.Cast<FLiveLinkCameraStaticData>();
+					FLiveLinkCameraFrameData* FrameData = SubjectData.FrameData.Cast<FLiveLinkCameraFrameData>();
+					if (StaticData->FIZDataMode == ECameraFIZMode::EncoderData)
+					{
+						if (StaticData->bIsFocusDistanceSupported)
+						{
+							CachedLiveLinkData.NormalizedFocus = FrameData->FocusDistance;
+							if (LensFile->HasFocusEncoderMapping())
+							{
+								CachedLiveLinkData.Focus = LensFile->EvaluateNormalizedFocus(FrameData->FocusDistance);
+							}
+						}
+
+						if (StaticData->bIsApertureSupported)
+						{
+							CachedLiveLinkData.NormalizedIris = FrameData->Aperture;
+							if (LensFile->HasIrisEncoderMapping())
+							{
+								CachedLiveLinkData.Iris = LensFile->EvaluateNormalizedIris(FrameData->Aperture);
+							}
+						}
+
+						if (StaticData->bIsFocalLengthSupported)
+						{
+							CachedLiveLinkData.NormalizedZoom = FrameData->FocalLength;
+							if (LensFile->HasZoomEncoderMapping())
+							{
+								CachedLiveLinkData.Zoom = LensFile->EvaluateNormalizedZoom(FrameData->FocalLength);
+							}
+						}
+					}
+					else
+					{
+						if (StaticData->bIsFocusDistanceSupported)
+						{
+							CachedLiveLinkData.Focus = FrameData->FocusDistance;
+						}
+
+						if (StaticData->bIsApertureSupported)
+						{
+							CachedLiveLinkData.Iris = FrameData->Aperture;
+						}
+
+						if (StaticData->bIsFocalLengthSupported)
+						{
+							CachedLiveLinkData.Zoom = FrameData->FocalLength;
+						}
+					}
+				}
+			}
+		}
+	}
+}
+
+void SLensEvaluation::CacheLensFileData()
+{
+	if (CachedLiveLinkData.Focus.IsSet() && CachedLiveLinkData.Zoom.IsSet())
+	{
+		const float Focus = CachedLiveLinkData.Focus.GetValue();
+		const float Zoom = CachedLiveLinkData.Zoom.GetValue();
+		LensFile->EvaluateDistortionParameters(Focus, Zoom, CachedDistortionInfo);
+		LensFile->EvaluateIntrinsicParameters(Focus, Zoom, CachedIntrinsics);
+		LensFile->EvaluateNodalPointOffset(Focus, Zoom, CachedNodalOffset);
+	}
+}
+
+TSharedRef<SWidget> SLensEvaluation::MakeTrackingWidget()
+{
+	return SNew(SVerticalBox)
+				+ SVerticalBox::Slot()
+				.Padding(5.0f, 5.0f)
+				.AutoHeight()
+				.HAlign(HAlign_Center)
+				[
+					SNew(STextBlock)
+					.Text(LOCTEXT("TrackingSection", "Tracking"))
+					.Font(FEditorStyle::GetFontStyle("DetailsView.CategoryFontStyle"))
+					.ShadowOffset(FVector2D(1.0f, 1.0f))
+				]
+				+ SVerticalBox::Slot()
+				.AutoHeight()
+				.HAlign(HAlign_Center)
+				[
+					SNew(SHorizontalBox)
+					+ SHorizontalBox::Slot()
+					.AutoWidth()
+					[
+						SNew(SCheckBox)
+						.ToolTipText(LOCTEXT("ViewModeTooltip", "Enable/Disable tracking usage"))
+						.IsChecked(this, &SLensEvaluation::IsTrackingActive)
+						.OnCheckStateChanged(this, &SLensEvaluation::OnTrackingStateChanged)
+					]
+					+ SHorizontalBox::Slot()
+					.AutoWidth()
+					[
+						SNew(SLiveLinkSubjectRepresentationPicker)
+						.ShowRole(false)
+						.Value(this, &SLensEvaluation::GetTrackingSubject)
+						.OnValueChanged(this, &SLensEvaluation::SetTrackingSubject)
+						.IsEnabled(this, &SLensEvaluation::CanSelectTrackingSource)
+					]
+				];
+}
+
+TSharedRef<SWidget> SLensEvaluation::MakeFIZWidget() const
+{
+	return SNew(SVerticalBox)
+		+ SVerticalBox::Slot()
+		.Padding(5.0f, 5.0f)
+		.AutoHeight()
+		.HAlign(HAlign_Center)
+		[
+			SNew(STextBlock)
+			.Text(LOCTEXT("FIZ Section", "FIZ"))
+			.Font(FEditorStyle::GetFontStyle("DetailsView.CategoryFontStyle"))
+			.ShadowOffset(FVector2D(1.0f, 1.0f))
+		]
+		+ SVerticalBox::Slot()
+		.AutoHeight()
+		[
+			SNew(SGridPanel)
+			.FillColumn(0, 0.2f)
+			.FillColumn(1, 0.4f)
+			.FillColumn(2, 0.4f)
+
+			+ SGridPanel::Slot(1, 0)
+			[
+				SNew(STextBlock)
+				.Text(LOCTEXT("EncodersLabel", "Encoders"))
+				.Font(FEditorStyle::GetFontStyle("DetailsView.CategoryFontStyle"))
+			]
+			+ SGridPanel::Slot(2, 0)
+			[
+				SNew(STextBlock)
+				.Text(LOCTEXT("PhysicalLabel", "Physicals"))
+				.Font(FEditorStyle::GetFontStyle("DetailsView.CategoryFontStyle"))
+			]
+
+			+ SGridPanel::Slot(0, 1)
+			[
+				SNew(STextBlock)
+				.Text(LOCTEXT("FocusLabel", "Focus"))
+				.Font(FEditorStyle::GetFontStyle("DetailsView.CategoryFontStyle"))
+			]
+			+ SGridPanel::Slot(0, 2)
+			[
+				SNew(STextBlock)
+				.Text(LOCTEXT("IrisLabel", "Iris"))
+				.Font(FEditorStyle::GetFontStyle("DetailsView.CategoryFontStyle"))
+			]
+			+ SGridPanel::Slot(0, 3)
+			[
+				SNew(STextBlock)
+				.Text(LOCTEXT("ZoomLabel", "Zoom"))
+				.Font(FEditorStyle::GetFontStyle("DetailsView.CategoryFontStyle"))
+			]
+
+			+ SGridPanel::Slot(1, 1)
+			[
+				SNew(STextBlock)
+				.Text(MakeAttributeLambda([this]
+				{
+					if (CachedLiveLinkData.NormalizedFocus.IsSet())
+					{
+						return FText::AsNumber(CachedLiveLinkData.NormalizedFocus.GetValue());
+					}
+					return LOCTEXT("UndefinedEncoderFocus", "N/A");
+				}))
+			]
+			+ SGridPanel::Slot(1, 2)
+			[
+				SNew(STextBlock)
+				.Text(MakeAttributeLambda([this]
+				{
+					if (CachedLiveLinkData.NormalizedIris.IsSet())
+					{
+						return FText::AsNumber(CachedLiveLinkData.NormalizedIris.GetValue());
+					}
+					return LOCTEXT("UndefinedEncoderIris", "N/A");
+				}))
+			]
+			+ SGridPanel::Slot(1, 3)
+			[
+				SNew(STextBlock)
+				.Text(MakeAttributeLambda([this]
+				{
+					if (CachedLiveLinkData.NormalizedZoom.IsSet())
+					{
+						return FText::AsNumber(CachedLiveLinkData.NormalizedZoom.GetValue());
+					}
+					return LOCTEXT("UndefinedEncoderZoom", "N/A");
+				}))
+			]
+
+			+ SGridPanel::Slot(2, 1)
+			[
+				SNew(STextBlock)
+				.Text(MakeAttributeLambda([this]
+				{
+					if (CachedLiveLinkData.Focus.IsSet())
+					{
+						return FText::AsNumber(CachedLiveLinkData.Focus.GetValue());
+					}
+					return LOCTEXT("UndefinedFocus", "N/A");
+				}))
+			]
+			+ SGridPanel::Slot(2, 2)
+			[
+				SNew(STextBlock)
+				.Text(MakeAttributeLambda([this]
+				{
+					if (CachedLiveLinkData.Iris.IsSet())
+					{
+						return FText::AsNumber(CachedLiveLinkData.Iris.GetValue());
+					}
+					return LOCTEXT("UndefinedIris", "N/A");
+				}))
+			]
+			+ SGridPanel::Slot(2, 3)
+			[
+				SNew(STextBlock)
+				.Text(MakeAttributeLambda([this]
+				{
+					if (CachedLiveLinkData.Zoom.IsSet())
+					{
+						return FText::AsNumber(CachedLiveLinkData.Zoom.GetValue());
+					}
+					return LOCTEXT("UndefinedZoom", "N/A");
+				}))
+			]
+		];
+}
+
+TSharedRef<SWidget> SLensEvaluation::MakeDistortionWidget() const
+{
+	//Find the named distortion parameters the current model has
+	TArray<FText> Parameters;
+	if (LensFile->LensInfo.LensModel)
+	{
+		Parameters = LensFile->LensInfo.LensModel.GetDefaultObject()->GetParameterDisplayNames();
+	}
+
+	const TSharedRef<SWidget> Title = SNew(STextBlock)
+		.Text(LOCTEXT("DistortionSection", "Distortion"))
+		.Font(FEditorStyle::GetFontStyle("DetailsView.CategoryFontStyle"))
+		.ShadowOffset(FVector2D(1.0f, 1.0f));
+
+	//if there are no parameters, create a simpler widget
+	if (Parameters.Num() <= 0)
+	{
+		return SNew(SVerticalBox)
+			+ SVerticalBox::Slot()
+			.Padding(5.0f, 5.0f)
+			.AutoHeight()
+			.HAlign(HAlign_Center)
+			[
+				Title
+			]
+			+ SVerticalBox::Slot()
+			.Padding(5.0f, 5.0f)
+			.AutoHeight()
+			.HAlign(HAlign_Center)
+			[
+				SNew(STextBlock)
+				.Text(LOCTEXT("NoParameters", "No parameters"))
+			];
+	}
+
+	TSharedRef<SGridPanel> ParameterGrid = SNew(SGridPanel);
+
+	for (int32 Index = 0; Index < Parameters.Num(); ++Index)
+	{
+		ParameterGrid->AddSlot(0, Index)
+			[
+				SNew(STextBlock)
+				.MinDesiredWidth(35.0f)
+				.Text(MakeAttributeLambda([Parameters, Index]
+				{
+					if (Parameters.IsValidIndex(Index))
+					{
+						return Parameters[Index];
+					}
+					else
+					{
+						return LOCTEXT("InvalidParam", "Invalid");
+					}
+				}))
+				.Font(FEditorStyle::GetFontStyle("DetailsView.CategoryFontStyle"))
+			];
+
+		ParameterGrid->AddSlot(1, Index)
+			[
+				SNew(STextBlock)
+				.Text(MakeAttributeLambda([this, Index]
+					{
+						if (CachedDistortionInfo.Parameters.IsValidIndex(Index))
+						{
+							return FText::AsNumber(CachedDistortionInfo.Parameters[Index]);
+						}
+						return LOCTEXT("UndefinedDistortionParameter", "N/A");
+					}))
+			];
+	}
+
+	return SNew(SVerticalBox)
+		+ SVerticalBox::Slot()
+		.Padding(5.0f, 5.0f)
+		.AutoHeight()
+		.HAlign(HAlign_Center)
+		[
+			Title
+		]
+		+ SVerticalBox::Slot()
+		.Padding(5.0f, 5.0f)
+		.AutoHeight()
+		.HAlign(HAlign_Center)
+		[
+			ParameterGrid
+		];
+}
+
+TSharedRef<SWidget> SLensEvaluation::MakeIntrinsicsWidget() const
+{
+	return
+		SNew(SVerticalBox)
+		+ SVerticalBox::Slot()
+		.Padding(5.0f, 5.0f)
+		.AutoHeight()
+		.HAlign(HAlign_Center)
+		[
+			SNew(STextBlock)
+			.Text(LOCTEXT("IntrinsicsSection", "Intrinsics"))
+			.Font(FEditorStyle::GetFontStyle("DetailsView.CategoryFontStyle"))
+			.ShadowOffset(FVector2D(1.0f, 1.0f))
+		]
+		+ SVerticalBox::Slot()
+		.AutoHeight()
+		.HAlign(HAlign_Center)
+		[
+			SNew(SGridPanel)
+
+			+ SGridPanel::Slot(0, 0)
+			[
+				SNew(STextBlock)
+				.Text(LOCTEXT("CxLabel", "Cx"))
+				.MinDesiredWidth(35.0f)
+				.Font(FEditorStyle::GetFontStyle("DetailsView.CategoryFontStyle"))
+			]
+			+ SGridPanel::Slot(1, 0)
+			[
+				SNew(STextBlock)
+				.Text(MakeAttributeLambda([this]
+				{
+					return FText::AsNumber(CachedIntrinsics.PrincipalPoint.X);
+				}))
+			]
+			+ SGridPanel::Slot(0, 1)
+			[
+				SNew(STextBlock)
+				.Text(LOCTEXT("CyLabel", "Cy"))
+				.MinDesiredWidth(35.0f)
+				.Font(FEditorStyle::GetFontStyle("DetailsView.CategoryFontStyle"))
+			]
+			+ SGridPanel::Slot(1, 1)
+			[
+				SNew(STextBlock)
+				.Text(MakeAttributeLambda([this]
+				{
+					return FText::AsNumber(CachedIntrinsics.PrincipalPoint.Y);
+				}))
+			]
+			+ SGridPanel::Slot(0, 2)
+			[
+				SNew(STextBlock)
+				.Text(LOCTEXT("FxLabel", "Fx"))
+				.MinDesiredWidth(35.0f)
+				.Font(FEditorStyle::GetFontStyle("DetailsView.CategoryFontStyle"))
+			]
+			+ SGridPanel::Slot(1, 2)
+			[
+				SNew(STextBlock)
+				.Text(MakeAttributeLambda([this]
+				{
+					return FText::AsNumber(CachedDistortionInfo.FxFy.X);
+				}))
+			]
+			+ SGridPanel::Slot(0, 3)
+			[
+				SNew(STextBlock)
+				.Text(LOCTEXT("FyLabel", "Fy"))
+				.MinDesiredWidth(35.0f)
+				.Font(FEditorStyle::GetFontStyle("DetailsView.CategoryFontStyle"))
+			]
+			+ SGridPanel::Slot(1, 3)
+			[
+				SNew(STextBlock)
+				.Text(MakeAttributeLambda([this]
+				{
+					return FText::AsNumber(CachedDistortionInfo.FxFy.Y);
+				}))
+			]
+		];
+}
+
+TSharedRef<SWidget> SLensEvaluation::MakeNodalOffsetWidget() const
+{
+	return
+		SNew(SVerticalBox)
+		+ SVerticalBox::Slot()
+		.Padding(5.0f, 5.0f)
+		.AutoHeight()
+		.HAlign(HAlign_Center)
+		[
+			SNew(STextBlock)
+			.Text(LOCTEXT("NodalOffsetSection", "Nodal Offset"))
+			.Font(FEditorStyle::GetFontStyle("DetailsView.CategoryFontStyle"))
+			.ShadowOffset(FVector2D(1.0f, 1.0f))
+		]
+		+ SVerticalBox::Slot()
+		.AutoHeight()
+		.HAlign(HAlign_Center)
+		[
+			SNew(SGridPanel)
+
+			+ SGridPanel::Slot(0, 0)
+			[
+				SNew(STextBlock)
+				.Text(LOCTEXT("PositionOffsetLabel", "Pos"))
+				.MinDesiredWidth(35.0f)
+				.Font(FEditorStyle::GetFontStyle("DetailsView.CategoryFontStyle"))
+			]
+			+ SGridPanel::Slot(1, 0)
+			[
+				SNew(STextBlock)
+				.MinDesiredWidth(15.0f)
+				.Text(MakeAttributeLambda([this]
+				{
+					return FText::AsNumber(CachedNodalOffset.LocationOffset.X);
+				}))
+			]
+			+ SGridPanel::Slot(2, 0)
+			[
+				SNew(STextBlock)
+				.MinDesiredWidth(15.0f)
+				.Text(MakeAttributeLambda([this]
+				{
+					return FText::AsNumber(CachedNodalOffset.LocationOffset.Y);
+				}))
+			]
+			+ SGridPanel::Slot(3, 0)
+			[
+				SNew(STextBlock)
+				.MinDesiredWidth(15.0f)
+				.Text(MakeAttributeLambda([this]
+				{
+					return FText::AsNumber(CachedNodalOffset.LocationOffset.Z);
+				}))
+			]
+
+			+ SGridPanel::Slot(0, 1)
+			[
+				SNew(STextBlock)
+				.Text(LOCTEXT("RotationOffsetLabel", "Rot"))
+				.MinDesiredWidth(35.0f)
+				.Font(FEditorStyle::GetFontStyle("DetailsView.CategoryFontStyle"))
+			]
+			+ SGridPanel::Slot(1, 1)
+			[
+				SNew(STextBlock)
+				.MinDesiredWidth(15.0f)
+				.Text(MakeAttributeLambda([this]
+				{
+					return FText::AsNumber(CachedNodalOffset.RotationOffset.X);
+				}))
+			]
+			+ SGridPanel::Slot(2, 1)
+			[
+				SNew(STextBlock)
+				.MinDesiredWidth(15.0f)
+				.Text(MakeAttributeLambda([this]
+				{
+					return FText::AsNumber(CachedNodalOffset.RotationOffset.Y);
+				}))
+			]
+			+ SGridPanel::Slot(3, 1)
+			[
+				SNew(STextBlock)
+				.MinDesiredWidth(15.0f)
+				.Text(MakeAttributeLambda([this]
+				{
+					return FText::AsNumber(CachedNodalOffset.RotationOffset.Z);
+				}))
+			]
+		];
+}
+
+#undef LOCTEXT_NAMESPACE /* LensDataViewer */
\ No newline at end of file
diff --git a/Engine/Plugins/VirtualProduction/CameraCalibration/Source/CameraCalibrationEditor/Private/AssetEditor/SLensEvaluation.h b/Engine/Plugins/VirtualProduction/CameraCalibration/Source/CameraCalibrationEditor/Private/AssetEditor/SLensEvaluation.h
new file mode 100644
--- /dev/null
+++ b/Engine/Plugins/VirtualProduction/CameraCalibration/Source/CameraCalibrationEditor/Private/AssetEditor/SLensEvaluation.h
@@ -0,0 +1,93 @@
+// Copyright Epic Games, Inc. All Rights Reserved.
+
+#pragma once
+
+#include "Widgets/SCompoundWidget.h"
+
+#include "LiveLinkRole.h"
+#include "SLensFilePanel.h"
+#include "SLiveLinkSubjectRepresentationPicker.h"
+#include "Widgets/DeclarativeSyntaxSupport.h"
+#include "UObject/StrongObjectPtr.h"
+
+
+class ULensFile;
+
+/**
+ * Widget using LiveLink subject input to evaluate lens file and show resulting data
+ */
+class SLensEvaluation : public SCompoundWidget
+{
+	using Super = SCompoundWidget;
+public:
+	SLATE_BEGIN_ARGS(SLensEvaluation)
+	{}
+	SLATE_END_ARGS()
+
+	void Construct(const FArguments& InArgs, ULensFile* InLensFile);
+
+	//~ Begin SCompoundWidget interface
+	virtual void Tick(const FGeometry& AllottedGeometry, const double InCurrentTime, const float InDeltaTime) override;
+	//~ End SCompoundWidget interface
+
+	/** Returns last evaluated FIZ data from LiveLink */
+	FCachedFIZData GetLastEvaluatedData() const { return CachedLiveLinkData; }
+
+private:
+
+	/** Returns tracking checkbox state */
+	ECheckBoxState IsTrackingActive() const;
+
+	/** Triggered when tracking checkbox state changes */
+	void OnTrackingStateChanged(ECheckBoxState NewState);
+
+	/** Whether tracking should be active and source can be selected */
+	bool CanSelectTrackingSource() const;
+
+	/** Returns the currently selected LiveLink subject */
+	SLiveLinkSubjectRepresentationPicker::FLiveLinkSourceSubjectRole GetTrackingSubject() const;
+
+	/** Set LiveLink subject to use */
+	void SetTrackingSubject(SLiveLinkSubjectRepresentationPicker::FLiveLinkSourceSubjectRole NewValue);
+
+	/** Evaluates LiveLink subjects */
+	void CacheLiveLinkData();
+
+	/** Evaluates LensFile using tracking data */
+	void CacheLensFileData();
+
+	/** Whether it should evaluate LiveLink subject */
+	bool ShouldUpdateTracking() const { return IsTrackingActive() == ECheckBoxState::Checked; }
+
+	/** Make LiveLink subject selection widget */
+	TSharedRef<SWidget> MakeTrackingWidget();
+	
+	/** Make FIZ widget evaluated from LiveLink subject */
+	TSharedRef<SWidget> MakeFIZWidget() const;
+
+	/** Make widget showing distortion data evaluated from LensFile */
+	TSharedRef<SWidget> MakeDistortionWidget() const;
+
+	/** Make widget showing Intrinsic data evaluated from LensFile */
+	TSharedRef<SWidget> MakeIntrinsicsWidget() const;
+
+	/** Make widget showing NodalOffset data evaluated from LensFile */
+	TSharedRef<SWidget> MakeNodalOffsetWidget() const;
+
+private:
+
+	/** LensFile being edited */
+	TStrongObjectPtr<ULensFile> LensFile;
+
+	/** Whether LiveLink is used */
+	bool bIsUsingLiveLinkTracking = true;
+
+	/** Subject used to get tracking input */
+	FLiveLinkSubjectRepresentation TrackingSource;
+
+	//~ Cached data taken from LiveLink subject and LensFile
+	FCachedFIZData CachedLiveLinkData;
+	FDistortionInfo CachedDistortionInfo;
+	FIntrinsicParameters CachedIntrinsics;
+	FNodalPointOffset CachedNodalOffset;
+};
\ No newline at end of file
diff --git a/Engine/Plugins/VirtualProduction/CameraCalibration/Source/CameraCalibrationEditor/Private/AssetEditor/SLensFilePanel.cpp b/Engine/Plugins/VirtualProduction/CameraCalibration/Source/CameraCalibrationEditor/Private/AssetEditor/SLensFilePanel.cpp
new file mode 100644
--- /dev/null
+++ b/Engine/Plugins/VirtualProduction/CameraCalibration/Source/CameraCalibrationEditor/Private/AssetEditor/SLensFilePanel.cpp
@@ -0,0 +1,47 @@
+// Copyright Epic Games, Inc. All Rights Reserved.
+
+#include "SLensFilePanel.h"
+
+#include "IStructureDetailsView.h"
+#include "LensFile.h"
+#include "Modules/ModuleManager.h"
+#include "PropertyEditorModule.h"
+#include "SLensDataViewer.h"
+#include "UObject/StructOnScope.h"
+
+
+#define LOCTEXT_NAMESPACE "LensFilePanel"
+
+
+void SLensFilePanel::Construct(const FArguments& InArgs, ULensFile* InLensFile)
+{
+	LensFile = TStrongObjectPtr<ULensFile>(InLensFile);
+	CachedFIZ = InArgs._CachedFIZData;
+
+	FStructureDetailsViewArgs LensInfoStructDetailsView;
+	FDetailsViewArgs DetailArgs;
+	DetailArgs.bAllowSearch = false;
+	DetailArgs.bShowScrollBar = false;
+
+	FPropertyEditorModule& PropertyEditor = FModuleManager::Get().LoadModuleChecked<FPropertyEditorModule>(TEXT("PropertyEditor"));
+	TSharedPtr<FStructOnScope> LensInfoStructScope = MakeShared<FStructOnScope>(FLensInfo::StaticStruct(), reinterpret_cast<uint8*>(&InLensFile->LensInfo));
+	TSharedPtr<IStructureDetailsView> LensInfoStructureDetailsView = PropertyEditor.CreateStructureDetailView(DetailArgs, LensInfoStructDetailsView, LensInfoStructScope);
+
+	ChildSlot
+	[
+		SNew(SHorizontalBox)
+		+ SHorizontalBox::Slot()
+		.FillWidth(0.8f)
+		[
+			SNew(SLensDataViewer, InLensFile)
+			.CachedFIZData(InArgs._CachedFIZData)
+		]
+		+ SHorizontalBox::Slot()
+		.FillWidth(0.2f)
+		[
+			LensInfoStructureDetailsView->GetWidget().ToSharedRef()
+		]
+	];
+}
+
+#undef LOCTEXT_NAMESPACE /*RemoteControlPanel*/
diff --git a/Engine/Plugins/VirtualProduction/CameraCalibration/Source/CameraCalibrationEditor/Private/AssetEditor/SLensFilePanel.h b/Engine/Plugins/VirtualProduction/CameraCalibration/Source/CameraCalibrationEditor/Private/AssetEditor/SLensFilePanel.h
new file mode 100644
--- /dev/null
+++ b/Engine/Plugins/VirtualProduction/CameraCalibration/Source/CameraCalibrationEditor/Private/AssetEditor/SLensFilePanel.h
@@ -0,0 +1,34 @@
+// Copyright Epic Games, Inc. All Rights Reserved.
+#pragma once
+
+#include "Widgets/SCompoundWidget.h"
+
+#include "CameraCalibrationEditorCommon.h"
+#include "UObject/StrongObjectPtr.h"
+
+class ULensFile;
+
+
+/**
+ * Main panel of lens file editor
+ */
+class SLensFilePanel : public SCompoundWidget
+{
+	SLATE_BEGIN_ARGS(SLensFilePanel) {}
+	
+		/** FIZ data */
+		SLATE_ATTRIBUTE(FCachedFIZData, CachedFIZData)
+
+	SLATE_END_ARGS()
+
+public:
+	void Construct(const FArguments& InArgs, ULensFile* InLensFile);
+
+private:
+
+	/** Holds the preset asset. */
+	TStrongObjectPtr<ULensFile> LensFile;
+
+	/** Evaluated FIZ for the current frame */
+	TAttribute<FCachedFIZData> CachedFIZ;
+};
diff --git a/Engine/Plugins/VirtualProduction/CameraCalibration/Source/CameraCalibrationEditor/Private/AssetEditor/SNodalOffsetToolPanel.cpp b/Engine/Plugins/VirtualProduction/CameraCalibration/Source/CameraCalibrationEditor/Private/AssetEditor/SNodalOffsetToolPanel.cpp
new file mode 100644
--- /dev/null
+++ b/Engine/Plugins/VirtualProduction/CameraCalibration/Source/CameraCalibrationEditor/Private/AssetEditor/SNodalOffsetToolPanel.cpp
@@ -0,0 +1,30 @@
+// Copyright Epic Games, Inc. All Rights Reserved.
+
+#include "SNodalOffsetToolPanel.h"
+
+#include "Widgets/SBoxPanel.h"
+#include "Widgets/Text/STextBlock.h"
+
+
+#define LOCTEXT_NAMESPACE "NodalOffsetTool"
+
+
+void SNodalOffsetToolPanel::Construct(const FArguments& InArgs, ULensFile* InLensFile)
+{
+	LensFile = TStrongObjectPtr<ULensFile>(InLensFile);
+
+	ChildSlot
+	[
+		SNew(SVerticalBox)
+		+ SVerticalBox::Slot()
+		.FillHeight(1.0f)
+		.VAlign(VAlign_Center)
+		[
+			SNew(STextBlock)
+			.Text(LOCTEXT("TEMPPLACEHOLDER", "Nothing to show"))
+		]
+	];
+}
+
+
+#undef LOCTEXT_NAMESPACE /*NodalOffsetTool*/
diff --git a/Engine/Plugins/VirtualProduction/CameraCalibration/Source/CameraCalibrationEditor/Private/AssetEditor/SNodalOffsetToolPanel.h b/Engine/Plugins/VirtualProduction/CameraCalibration/Source/CameraCalibrationEditor/Private/AssetEditor/SNodalOffsetToolPanel.h
new file mode 100644
--- /dev/null
+++ b/Engine/Plugins/VirtualProduction/CameraCalibration/Source/CameraCalibrationEditor/Private/AssetEditor/SNodalOffsetToolPanel.h
@@ -0,0 +1,26 @@
+// Copyright Epic Games, Inc. All Rights Reserved.
+#pragma once
+
+#include "Widgets/SCompoundWidget.h"
+
+#include "LensFile.h"
+#include "UObject/StrongObjectPtr.h"
+
+
+/**
+ * Placeholder for camera calibration tools tab
+ */
+class SNodalOffsetToolPanel : public SCompoundWidget
+{
+	SLATE_BEGIN_ARGS(SNodalOffsetToolPanel) {}
+	SLATE_END_ARGS()
+
+public:
+	void Construct(const FArguments& InArgs, ULensFile* InLensFile);
+
+
+private:
+
+	/** Holds the preset asset. */
+	TStrongObjectPtr<ULensFile> LensFile;
+};
diff --git a/Engine/Plugins/VirtualProduction/CameraCalibration/Source/CameraCalibrationEditor/Private/AssetTypeActions/AssetTypeActions_LensFile.cpp b/Engine/Plugins/VirtualProduction/CameraCalibration/Source/CameraCalibrationEditor/Private/AssetTypeActions/AssetTypeActions_LensFile.cpp
--- a/Engine/Plugins/VirtualProduction/CameraCalibration/Source/CameraCalibrationEditor/Private/AssetTypeActions/AssetTypeActions_LensFile.cpp
+++ b/Engine/Plugins/VirtualProduction/CameraCalibration/Source/CameraCalibrationEditor/Private/AssetTypeActions/AssetTypeActions_LensFile.cpp
@@ -2,7 +2,7 @@
 
 #include "AssetTypeActions/AssetTypeActions_LensFile.h"
 
-#include "AssetEditor/LensFileEditorToolkit.h"
+#include "AssetEditor/CameraCalibrationToolkit.h"
 #include "LensFile.h"
 
 #define LOCTEXT_NAMESPACE "LensFileTypeActions"
@@ -17,15 +17,17 @@ UClass* FAssetTypeActions_LensFile::GetSupportedClass() const
 	return ULensFile::StaticClass();
 }
 
-//void FAssetTypeActions_LensFile::OpenAssetEditor(const TArray<UObject*>& InObjects, TSharedPtr<IToolkitHost> EditWithinLevelEditor)
-//{
-//	for (UObject* Object : InObjects)
-//	{
-//		if (ULensFile* Asset = Cast<ULensFile>(Object))
-//		{
-//			FLensFileEditorToolkit::CreateEditor(EToolkitMode::Standalone, EditWithinLevelEditor, Asset);
-//		}
-//	}
-//}
+void FAssetTypeActions_LensFile::OpenAssetEditor(const TArray<UObject*>& InObjects, TSharedPtr<IToolkitHost> EditWithinLevelEditor)
+{
+	const EToolkitMode::Type Mode = EditWithinLevelEditor.IsValid() ? EToolkitMode::WorldCentric : EToolkitMode::Standalone;
+
+	for (UObject* Object : InObjects)
+	{
+		if (ULensFile* Asset = Cast<ULensFile>(Object))
+		{
+			FCameraCalibrationToolkit::CreateEditor(Mode, EditWithinLevelEditor, Asset);
+		}
+	}
+}
 
 #undef LOCTEXT_NAMESPACE
diff --git a/Engine/Plugins/VirtualProduction/CameraCalibration/Source/CameraCalibrationEditor/Private/AssetTypeActions/AssetTypeActions_LensFile.h b/Engine/Plugins/VirtualProduction/CameraCalibration/Source/CameraCalibrationEditor/Private/AssetTypeActions/AssetTypeActions_LensFile.h
--- a/Engine/Plugins/VirtualProduction/CameraCalibration/Source/CameraCalibrationEditor/Private/AssetTypeActions/AssetTypeActions_LensFile.h
+++ b/Engine/Plugins/VirtualProduction/CameraCalibration/Source/CameraCalibrationEditor/Private/AssetTypeActions/AssetTypeActions_LensFile.h
@@ -16,6 +16,6 @@ public:
 	virtual bool IsImportedAsset() const override { return false; }
 	virtual FText GetName() const override;
 	virtual UClass* GetSupportedClass() const override;
-	//virtual void OpenAssetEditor(const TArray<UObject*>& InObjects, TSharedPtr<IToolkitHost> EditWithinLevelEditor = TSharedPtr<IToolkitHost>()) override;
+	virtual void OpenAssetEditor(const TArray<UObject*>& InObjects, TSharedPtr<IToolkitHost> EditWithinLevelEditor = TSharedPtr<IToolkitHost>()) override;
 	//~ End IAssetTypeActions interface
 };
\ No newline at end of file
diff --git a/Engine/Plugins/VirtualProduction/LiveLinkCamera/Source/LiveLinkCamera/Private/LiveLinkCameraController.cpp b/Engine/Plugins/VirtualProduction/LiveLinkCamera/Source/LiveLinkCamera/Private/LiveLinkCameraController.cpp
--- a/Engine/Plugins/VirtualProduction/LiveLinkCamera/Source/LiveLinkCamera/Private/LiveLinkCameraController.cpp
+++ b/Engine/Plugins/VirtualProduction/LiveLinkCamera/Source/LiveLinkCamera/Private/LiveLinkCameraController.cpp
@@ -215,10 +215,9 @@ void ULiveLinkCameraController::ApplyFIZ(ULensFile* LensFile, UCineCameraCompone
 		{
 			if (StaticData->bIsFocusDistanceSupported && UpdateFlags.bApplyFocusDistance)
 			{
-				float NewFocusDistance;
-				if (LensFile->EvaluateNormalizedFocus(FrameData->FocusDistance, NewFocusDistance))
+				if (LensFile->HasFocusEncoderMapping())
 				{
-					CineCameraComponent->FocusSettings.ManualFocusDistance = NewFocusDistance;
+					CineCameraComponent->FocusSettings.ManualFocusDistance = LensFile->EvaluateNormalizedFocus(FrameData->FocusDistance);
 				}
 				else
 				{
@@ -228,10 +227,9 @@ void ULiveLinkCameraController::ApplyFIZ(ULensFile* LensFile, UCineCameraCompone
 
 			if (StaticData->bIsApertureSupported && UpdateFlags.bApplyAperture)
 			{
-				float NewAperture;
-				if (LensFile->EvaluateNormalizedIris(FrameData->Aperture, NewAperture))
+				if (LensFile->HasIrisEncoderMapping())
 				{
-					CineCameraComponent->CurrentAperture = NewAperture;
+					CineCameraComponent->CurrentAperture = LensFile->EvaluateNormalizedIris(FrameData->Aperture);
 				}
 				else
 				{
@@ -241,10 +239,9 @@ void ULiveLinkCameraController::ApplyFIZ(ULensFile* LensFile, UCineCameraCompone
 
 			if (StaticData->bIsFocalLengthSupported && UpdateFlags.bApplyFocalLength)
 			{
-				float NewZoom;
-				if (LensFile->EvaluateNormalizedZoom(FrameData->FocalLength, NewZoom))
+				if (LensFile->HasZoomEncoderMapping())
 				{
-					CineCameraComponent->SetCurrentFocalLength(NewZoom);
+					CineCameraComponent->SetCurrentFocalLength(LensFile->EvaluateNormalizedZoom(FrameData->FocalLength));
 				}
 				else
 				{
