CameraCalibration: Added Simulcam Viewer and a Nodal Offset calibrator.

diff --git a/Engine/Plugins/VirtualProduction/CameraCalibration/CameraCalibration.uplugin b/Engine/Plugins/VirtualProduction/CameraCalibration/CameraCalibration.uplugin
--- a/Engine/Plugins/VirtualProduction/CameraCalibration/CameraCalibration.uplugin
+++ b/Engine/Plugins/VirtualProduction/CameraCalibration/CameraCalibration.uplugin
@@ -29,6 +29,18 @@
 		{
 			"Name": "LiveLink",
 			"Enabled": true
+		},
+		{
+			"Name": "Composure",
+			"Enabled": true
+		},
+		{
+			"Name": "MediaFrameworkUtilities",
+			"Enabled": true
+		},
+		{
+			"Name": "OpenCVLensDistortion",
+			"Enabled": true
 		}
 	]
 }
\ No newline at end of file
diff --git a/Engine/Plugins/VirtualProduction/CameraCalibration/Source/CameraCalibration/CameraCalibration.Build.cs b/Engine/Plugins/VirtualProduction/CameraCalibration/Source/CameraCalibration/CameraCalibration.Build.cs
--- a/Engine/Plugins/VirtualProduction/CameraCalibration/Source/CameraCalibration/CameraCalibration.Build.cs
+++ b/Engine/Plugins/VirtualProduction/CameraCalibration/Source/CameraCalibration/CameraCalibration.Build.cs
@@ -6,6 +6,14 @@ namespace UnrealBuildTool.Rules
 	{
 		public CameraCalibration(ReadOnlyTargetRules Target) : base(Target)
 		{
+			if (Target.Configuration == UnrealTargetConfiguration.Debug)
+			{
+				PCHUsage = ModuleRules.PCHUsageMode.UseExplicitOrSharedPCHs;
+				OptimizeCode = CodeOptimization.Never;
+				bUseUnity = false;
+				PCHUsage = PCHUsageMode.NoPCHs;
+			}
+
 			PublicDependencyModuleNames.AddRange(
 				new string[]
 				{
@@ -16,6 +24,8 @@ namespace UnrealBuildTool.Rules
 					"Engine",
 					"RenderCore",
 					"RHI",
+					"Slate",
+					"SlateCore",
 				}
 			);
 
diff --git a/Engine/Plugins/VirtualProduction/CameraCalibration/Source/CameraCalibration/Private/CameraCalibrationSubsystem.cpp b/Engine/Plugins/VirtualProduction/CameraCalibration/Source/CameraCalibration/Private/CameraCalibrationSubsystem.cpp
--- a/Engine/Plugins/VirtualProduction/CameraCalibration/Source/CameraCalibration/Private/CameraCalibrationSubsystem.cpp
+++ b/Engine/Plugins/VirtualProduction/CameraCalibration/Source/CameraCalibration/Private/CameraCalibrationSubsystem.cpp
@@ -5,6 +5,7 @@
 #include "CameraCalibrationLog.h"
 #include "CineCameraComponent.h"
 #include "Engine/TimecodeProvider.h"
+#include "Misc/CoreDelegates.h"
 #include "UObject/UObjectIterator.h"
 
 #if WITH_EDITOR
@@ -133,7 +134,6 @@ void UCameraCalibrationSubsystem::UnregisterDistortionModelHandler(UCineCameraCo
 	LensDistortionHandlerMap.Remove(Component, Handler);
 }
 
-
 void UCameraCalibrationSubsystem::RegisterDistortionModel(TSubclassOf<ULensModel> LensModel)
 {
 	LensModelMap.Add(LensModel->GetDefaultObject<ULensModel>()->GetModelName(), LensModel);
@@ -153,9 +153,82 @@ TSubclassOf<ULensModel> UCameraCalibrationSubsystem::GetRegisteredLensModel(FNam
 	return nullptr;
 }
 
+TSubclassOf<UCameraNodalOffsetAlgo> UCameraCalibrationSubsystem::GetCameraNodalOffsetAlgo(FName Name) const
+{
+	if (CameraNodalOffsetAlgosMap.Contains(Name))
+	{
+		return CameraNodalOffsetAlgosMap[Name];
+	}
+	return nullptr;
+}
+
+TArray<FName> UCameraCalibrationSubsystem::GetCameraNodalOffsetAlgos() const
+{
+	TArray<FName> OutKeys;
+	CameraNodalOffsetAlgosMap.GetKeys(OutKeys);
+	return OutKeys;
+}
+
+TSubclassOf<UCameraCalibrationStep> UCameraCalibrationSubsystem::GetCameraCalibrationStep(FName Name) const
+{
+	if (CameraCalibrationStepsMap.Contains(Name))
+	{
+		return CameraCalibrationStepsMap[Name];
+	}
+	return nullptr;
+}
+
+TArray<FName> UCameraCalibrationSubsystem::GetCameraCalibrationSteps() const
+{
+	TArray<FName> OutKeys;
+	CameraCalibrationStepsMap.GetKeys(OutKeys);
+	return OutKeys;
+}
+
+void UCameraCalibrationSubsystem::Initialize(FSubsystemCollectionBase& Collection)
+{
+	Super::Initialize(Collection);
+
+	PostEngineInitHandle = FCoreDelegates::OnPostEngineInit.AddLambda([&]()
+	{
+		// Find Nodal Offset Algos
+		{
+			TArray<TSubclassOf<UCameraNodalOffsetAlgo>> Algos;
+
+			for (TObjectIterator<UClass> AlgoIt; AlgoIt; ++AlgoIt)
+			{
+				if (AlgoIt->IsChildOf(UCameraNodalOffsetAlgo::StaticClass()) && !AlgoIt->HasAnyClassFlags(CLASS_Abstract | CLASS_Deprecated))
+				{
+					CameraNodalOffsetAlgosMap.Add(AlgoIt->GetFName(), TSubclassOf<UCameraNodalOffsetAlgo>(*AlgoIt));
+				}
+			}
+		}
+
+		// Find Calibration Steps
+		{
+			TArray<TSubclassOf<UCameraCalibrationStep>> Steps;
+
+			for (TObjectIterator<UClass> StepIt; StepIt; ++StepIt)
+			{
+				if (StepIt->IsChildOf(UCameraCalibrationStep::StaticClass()) && !StepIt->HasAnyClassFlags(CLASS_Abstract | CLASS_Deprecated))
+				{
+					CameraCalibrationStepsMap.Add(StepIt->GetFName(), TSubclassOf<UCameraCalibrationStep>(*StepIt));
+				}
+			}
+		}
+	});
+}
+
 void UCameraCalibrationSubsystem::Deinitialize()
 {
 	LensModelMap.Empty(0);
+	CameraNodalOffsetAlgosMap.Empty(0);
+	CameraCalibrationStepsMap.Empty(0);
+
+	if (PostEngineInitHandle.IsValid())
+	{
+		FCoreDelegates::OnPostEngineInit.Remove(PostEngineInitHandle);
+	}
 
 	Super::Deinitialize();
 }
diff --git a/Engine/Plugins/VirtualProduction/CameraCalibration/Source/CameraCalibration/Private/CameraCalibrationUtils.cpp b/Engine/Plugins/VirtualProduction/CameraCalibration/Source/CameraCalibration/Private/CameraCalibrationUtils.cpp
new file mode 100644
--- /dev/null
+++ b/Engine/Plugins/VirtualProduction/CameraCalibration/Source/CameraCalibration/Private/CameraCalibrationUtils.cpp
@@ -0,0 +1,62 @@
+// Copyright Epic Games, Inc. All Rights Reserved.
+
+
+#include "CameraCalibrationUtils.h"
+
+#include <type_traits>
+
+
+namespace CameraCalibrationUtils
+{
+	using EAxis = FCameraCalibrationUtils::EAxis;
+
+	// These axes must match the order in which they are declared in EAxis
+	static const TArray<FVector> UnitVectors =
+	{
+		{  1,  0,  0 }, //  X
+		{  0,  1,  0 }, //  Y
+		{  0,  0,  1 }, //  Z
+		{ -1,  0,  0 }, // -X
+		{  0, -1,  0 }, // -Y
+		{  0,  0, -1 }, // -Z
+	};
+
+	static const FVector& UnitVectorFromAxisEnum(EAxis Axis)
+	{
+		return UnitVectors[std::underlying_type_t<EAxis>(Axis)];
+	};
+}
+
+void FCameraCalibrationUtils::ConvertCoordinateSystem(FTransform& Transform, EAxis SrcXInDstAxis, EAxis SrcYInDstAxis, EAxis SrcZInDstAxis)
+{
+	// Unreal Engine:
+	//   Front : X
+	//   Right : Y
+	//   Up    : Z
+	//
+	// OpenCV:
+	//   Front : Z
+	//   Right : X
+	//   Up    : Yn
+
+	using namespace CameraCalibrationUtils;
+
+	FMatrix M12 = FMatrix::Identity;
+
+	M12.SetColumn(0, UnitVectorFromAxisEnum(SrcXInDstAxis));
+	M12.SetColumn(1, UnitVectorFromAxisEnum(SrcYInDstAxis));
+	M12.SetColumn(2, UnitVectorFromAxisEnum(SrcZInDstAxis));
+
+	Transform.SetFromMatrix(M12.GetTransposed() * Transform.ToMatrixWithScale() * M12);
+}
+
+void FCameraCalibrationUtils::ConvertUnrealToOpenCV(FTransform& Transform)
+{
+	ConvertCoordinateSystem(Transform, EAxis::Y, EAxis::Zn, EAxis::X);
+}
+
+void FCameraCalibrationUtils::ConvertOpenCVToUnreal(FTransform& Transform)
+{
+	ConvertCoordinateSystem(Transform, EAxis::Z, EAxis::X, EAxis::Yn);
+}
+
diff --git a/Engine/Plugins/VirtualProduction/CameraCalibration/Source/CameraCalibration/Private/LensFile.cpp b/Engine/Plugins/VirtualProduction/CameraCalibration/Source/CameraCalibration/Private/LensFile.cpp
--- a/Engine/Plugins/VirtualProduction/CameraCalibration/Source/CameraCalibration/Private/LensFile.cpp
+++ b/Engine/Plugins/VirtualProduction/CameraCalibration/Source/CameraCalibration/Private/LensFile.cpp
@@ -130,7 +130,7 @@ void ULensFile::PostEditChangeChainProperty(struct FPropertyChangedChainEvent& P
 
 #endif //WITH_EDITOR
 
-bool ULensFile::EvaluateDistortionParameters(float InFocus, float InZoom, FDistortionInfo& OutEvaluatedValue)
+bool ULensFile::EvaluateDistortionParameters(float InFocus, float InZoom, FDistortionInfo& OutEvaluatedValue) const
 {
 	if (DistortionMapping.Num() <= 0)
 	{
@@ -153,7 +153,7 @@ bool ULensFile::EvaluateDistortionParameters(float InFocus, float InZoom, FDisto
 	return bSuccess;
 }
 
-bool ULensFile::EvaluateIntrinsicParameters(float InFocus, float InZoom, FIntrinsicParameters& OutEvaluatedValue)
+bool ULensFile::EvaluateIntrinsicParameters(float InFocus, float InZoom, FIntrinsicParameters& OutEvaluatedValue) const
 {
 	if (IntrinsicMapping.Num() <= 0)
 	{
@@ -176,7 +176,7 @@ bool ULensFile::EvaluateIntrinsicParameters(float InFocus, float InZoom, FIntrin
 	return bSuccess;
 }
 
-bool ULensFile::EvaluateDistortionData(float InFocus, float InZoom, FVector2D InFilmback, ULensDistortionModelHandlerBase* InLensHandler, FDistortionData& OutDistortionData)
+bool ULensFile::EvaluateDistortionData(float InFocus, float InZoom, FVector2D InFilmback, ULensDistortionModelHandlerBase* InLensHandler, FDistortionData& OutDistortionData) const
 {	
 	if (InLensHandler == nullptr)
 	{
@@ -257,7 +257,7 @@ void ULensFile::SetupNoDistortionOutput(ULensDistortionModelHandlerBase* LensHan
 	LensHandler->SetOverscanFactor(OutDistortionData.OverscanFactor);
 }
 
-bool ULensFile::EvaluateDistortionForParameters(float InFocus, float InZoom, FVector2D InFilmback, ULensDistortionModelHandlerBase* InLensHandler, FDistortionData& OutDistortionData)
+bool ULensFile::EvaluateDistortionForParameters(float InFocus, float InZoom, FVector2D InFilmback, ULensDistortionModelHandlerBase* InLensHandler, FDistortionData& OutDistortionData) const
 {
 	//Parameter blending mode or no parameters to blend
 	if (LensFileUtils::GDistortionParametersBlendMode == 0
@@ -304,10 +304,10 @@ bool ULensFile::EvaluateDistortionForParameters(float InFocus, float InZoom, FVe
 				return false;	
 			}
 
-			FDistortionMapPoint& MinMinPoint = DistortionMapping[MinMinIndex];
-			FDistortionMapPoint& MinMaxPoint = DistortionMapping[MinMaxIndex];
-			FDistortionMapPoint& MaxMinPoint = DistortionMapping[MaxMinIndex];
-			FDistortionMapPoint& MaxMaxPoint = DistortionMapping[MaxMaxIndex];
+			const FDistortionMapPoint& MinMinPoint = DistortionMapping[MinMinIndex];
+			const FDistortionMapPoint& MinMaxPoint = DistortionMapping[MinMaxIndex];
+			const FDistortionMapPoint& MaxMinPoint = DistortionMapping[MaxMinIndex];
+			const FDistortionMapPoint& MaxMaxPoint = DistortionMapping[MaxMaxIndex];
 
 			FDistortionData BlendedData;
 			BlendedData.DistortedUVs.SetNumZeroed(UndistortedUVs.Num());
@@ -442,7 +442,7 @@ bool ULensFile::EvaluateDistortionForParameters(float InFocus, float InZoom, FVe
 	return true;
 }
 
-bool ULensFile::EvaluteDistortionForSTMaps(float InFocus, float InZoom, FVector2D InFilmback, ULensDistortionModelHandlerBase* InLensHandler, FDistortionData& OutDistortionData)
+bool ULensFile::EvaluteDistortionForSTMaps(float InFocus, float InZoom, FVector2D InFilmback, ULensDistortionModelHandlerBase* InLensHandler, FDistortionData& OutDistortionData) const
 {
 	if(DerivedDataInFlightCount > 0)
 	{
@@ -603,7 +603,7 @@ bool ULensFile::EvaluteDistortionForSTMaps(float InFocus, float InZoom, FVector2
 	return true;
 }
 
-bool ULensFile::EvaluateNodalPointOffset(float InFocus, float InZoom, FNodalPointOffset& OutEvaluatedValue)
+bool ULensFile::EvaluateNodalPointOffset(float InFocus, float InZoom, FNodalPointOffset& OutEvaluatedValue) const
 {
 	if (NodalOffsetMapping.Num() <= 0)
 	{
diff --git a/Engine/Plugins/VirtualProduction/CameraCalibration/Source/CameraCalibration/Public/CalibrationPointComponent.h b/Engine/Plugins/VirtualProduction/CameraCalibration/Source/CameraCalibration/Public/CalibrationPointComponent.h
new file mode 100644
--- /dev/null
+++ b/Engine/Plugins/VirtualProduction/CameraCalibration/Source/CameraCalibration/Public/CalibrationPointComponent.h
@@ -0,0 +1,19 @@
+// Copyright Epic Games, Inc. All Rights Reserved.
+
+#pragma once
+
+#include "Components/SceneComponent.h"
+
+#include "CalibrationPointComponent.generated.h"
+
+/**
+ * One or more instances of this component can be added to an actor (e.g. a static mesh actor blueprint), 
+ * and should be placed at geometrically and visually distinct landmarks of the object.
+ * These 3d points will then be optionally used by any given nodal offset tool implementation to
+ * make a 3d-2d correspondence with the 2d points detected in the live action media.
+ */
+UCLASS(ClassGroup = (Calibration), meta = (BlueprintSpawnableComponent), meta = (DisplayName = "Calibration Point"))
+class CAMERACALIBRATION_API UCalibrationPointComponent : public USceneComponent
+{
+	GENERATED_BODY()
+};
\ No newline at end of file
diff --git a/Engine/Plugins/VirtualProduction/CameraCalibration/Source/CameraCalibration/Public/CameraCalibrationStep.h b/Engine/Plugins/VirtualProduction/CameraCalibration/Source/CameraCalibration/Public/CameraCalibrationStep.h
new file mode 100644
--- /dev/null
+++ b/Engine/Plugins/VirtualProduction/CameraCalibration/Source/CameraCalibration/Public/CameraCalibrationStep.h
@@ -0,0 +1,53 @@
+// Copyright Epic Games, Inc. All Rights Reserved.
+
+#pragma once
+
+#include "CoreMinimal.h"
+#include "Templates/SharedPointer.h"
+#include "Widgets/SBoxPanel.h"
+
+#include "CameraCalibrationStep.generated.h"
+
+struct FGeometry;
+struct FPointerEvent;
+
+class FCameraCalibrationStepsController;
+
+/**
+ * Interface of a camera calibration step. These will appear in a Camera Calibration Toolkit tab.
+ */
+UCLASS(Abstract)
+class CAMERACALIBRATION_API UCameraCalibrationStep : public UObject
+{
+	GENERATED_BODY()
+
+public:
+
+	/** Make sure you initialize before using the object */
+	virtual void Initialize(FCameraCalibrationStepsController* InCameraCalibrationStepController) {};
+
+	/** Clean up resources and don't use CameraCalibrationStepController anymore */
+	virtual void Shutdown() {};
+
+	/** Called every frame */
+	virtual void Tick(float DeltaTime) {};
+
+	/** Callback when viewport is clicked. Returns false if the event was not handled. */
+	virtual bool OnViewportClicked(const FGeometry& MyGeometry, const FPointerEvent& MouseEvent) { return false;  };
+
+	/** Returns the UI of this camera calibration step. Expected to only be called once */
+	virtual TSharedRef<SWidget> BuildUI() { return SNew(SVerticalBox); };
+
+	/** Returns a title or friendly name that can be placed in UI */
+	virtual FName FriendlyName() const { return TEXT("Invalid"); };
+
+	/** Returns true if the given calibration step is a known prerequisite for this step */
+	virtual bool DependsOnStep(UCameraCalibrationStep* Step) const { return false; };
+
+	/** Called when this step is the active step in the UI */
+	virtual void Activate() {};
+
+	/** Called when this step is no longer the active step in the UI */
+	virtual void Deactivate() {};
+
+};
diff --git a/Engine/Plugins/VirtualProduction/CameraCalibration/Source/CameraCalibration/Public/CameraCalibrationSubsystem.h b/Engine/Plugins/VirtualProduction/CameraCalibration/Source/CameraCalibration/Public/CameraCalibrationSubsystem.h
--- a/Engine/Plugins/VirtualProduction/CameraCalibration/Source/CameraCalibration/Public/CameraCalibrationSubsystem.h
+++ b/Engine/Plugins/VirtualProduction/CameraCalibration/Source/CameraCalibration/Public/CameraCalibrationSubsystem.h
@@ -4,7 +4,12 @@
 
 #include "Subsystems/EngineSubsystem.h"
 
+#include "CoreMinimal.h"
+
 #include "CameraCalibrationTypes.h"
+#include "CameraNodalOffsetAlgo.h"
+#include "CameraCalibrationStep.h"
+#include "Containers/ArrayView.h"
 #include "LensDistortionModelHandlerBase.h"
 #include "LensFile.h"
 #include "UObject/ObjectKey.h"
@@ -58,9 +63,26 @@ public:
 	UFUNCTION(BlueprintCallable, Category = "Lens Distortion")
 	TSubclassOf<ULensModel> GetRegisteredLensModel(FName ModelName) const;
 
+	/** Returns the nodal offset algorithm by name */
+	UFUNCTION(BlueprintCallable, Category = "Lens Distortion")
+	TSubclassOf<UCameraNodalOffsetAlgo> GetCameraNodalOffsetAlgo(FName Name) const;
+
+	/** Returns an array with the names of the available nodal offset algorithms */
+	UFUNCTION(BlueprintCallable, Category = "Lens Distortion")
+	TArray<FName> GetCameraNodalOffsetAlgos() const;
+
+	/** Returns an array with the names of the available camera calibration steps */
+	UFUNCTION(BlueprintCallable, Category = "Lens Distortion")
+	TArray<FName> GetCameraCalibrationSteps() const;
+
+	/** Returns the camera calibration step by name */
+	UFUNCTION(BlueprintCallable, Category = "Lens Distortion")
+	TSubclassOf<UCameraCalibrationStep> GetCameraCalibrationStep(FName Name) const;
+
 public:
 	//~ Begin USubsystem interface
 	virtual void Deinitialize() override;
+	virtual void Initialize(FSubsystemCollectionBase& Collection) override;
 	//~ End USubsystem interface
 
 	/** Add a new Lens Model to the registered model map */
@@ -78,10 +100,22 @@ private:
 	UPROPERTY(Transient)
 	TMap<FName, TSubclassOf<ULensModel>> LensModelMap;
 
+	/** Holds the registered camera nodal offset algos */
+	UPROPERTY(Transient)
+	TMap<FName, TSubclassOf<UCameraNodalOffsetAlgo>> CameraNodalOffsetAlgosMap;
+
+	/** Holds the registered camera calibration steps */
+	UPROPERTY(Transient)
+	TMap<FName, TSubclassOf<UCameraCalibrationStep>> CameraCalibrationStepsMap;
+
 	/** Map of actor components to the authoritative lens model that should be used with that component */
 	TMap<FObjectKey, TSubclassOf<ULensModel>> ComponentsWithAuthoritativeModels;
 
 	/** Map of actor components to the authoritative lens model that should be used with that component */
 	TMultiMap<FObjectKey, ULensDistortionModelHandlerBase*> LensDistortionHandlerMap;
 
+private:
+
+	FDelegateHandle PostEngineInitHandle;
+
 };
diff --git a/Engine/Plugins/VirtualProduction/CameraCalibration/Source/CameraCalibration/Public/CameraCalibrationUtils.h b/Engine/Plugins/VirtualProduction/CameraCalibration/Source/CameraCalibration/Public/CameraCalibrationUtils.h
new file mode 100644
--- /dev/null
+++ b/Engine/Plugins/VirtualProduction/CameraCalibration/Source/CameraCalibration/Public/CameraCalibrationUtils.h
@@ -0,0 +1,33 @@
+// Copyright Epic Games, Inc. All Rights Reserved.
+
+#pragma once
+
+#include "CoreMinimal.h"
+
+struct FTransform;
+
+/**
+ * Helper class for commonly used functions for camera calibration.
+ */
+class CAMERACALIBRATION_API FCameraCalibrationUtils
+{
+public:
+
+	/** Enumeration to specify any cartesian axis in positive or negative directions */
+	enum class EAxis
+	{
+		X, Y, Z,
+		Xn, Yn, Zn,
+	};
+
+public:
+
+	/** Converts in-place the coordinate system of the given FTransform by specifying the source axes in terms of the destionation axes */
+	static void ConvertCoordinateSystem(FTransform& Transform, EAxis DstXInSrcAxis, EAxis DstYInSrcAxis, EAxis DstZInSrcAxis);
+
+	/** Converts in-place an FTransform in Unreal coordinates to OpenCV coordinates */
+	static void ConvertUnrealToOpenCV(FTransform& Transform);
+
+	/** Converts in-place an FTransform in OpenCV coordinates to Unreal coordinates */
+	static void ConvertOpenCVToUnreal(FTransform& Transform);
+};
diff --git a/Engine/Plugins/VirtualProduction/CameraCalibration/Source/CameraCalibration/Public/CameraNodalOffsetAlgo.h b/Engine/Plugins/VirtualProduction/CameraCalibration/Source/CameraCalibration/Public/CameraNodalOffsetAlgo.h
new file mode 100644
--- /dev/null
+++ b/Engine/Plugins/VirtualProduction/CameraCalibration/Source/CameraCalibration/Public/CameraNodalOffsetAlgo.h
@@ -0,0 +1,66 @@
+// Copyright Epic Games, Inc. All Rights Reserved.
+
+#pragma once
+
+#include "CoreMinimal.h"
+#include "Widgets/SBoxPanel.h"
+
+#include "CameraNodalOffsetAlgo.generated.h"
+
+struct FGeometry;
+struct FNodalPointOffset;
+struct FPointerEvent;
+struct FTransform;
+
+class FNodalOffsetTool;
+class SWidget;
+
+/**
+ * UCameraNodalOffsetAlgo defines the interface that any nodal calibration point algorithm should implement
+ * in order to be used and listed by the Nodal Offset Tool. 
+ */
+UCLASS(Abstract)
+class CAMERACALIBRATION_API UCameraNodalOffsetAlgo : public UObject
+{
+	GENERATED_BODY()
+
+public:
+
+	/** Structure with the camera information needed by the algorithm */
+	struct FCameraData
+	{
+		FTransform Pose;
+		uint32 CameraId;
+		float Zoom;
+		float Focus;
+		float FilmbackX;
+		float FilmbackY;
+		bool bValid;
+	};
+
+public:
+
+	/** Make sure you initialize before using the object */
+	virtual void Initialize(TWeakPtr<FNodalOffsetTool> InNodalOffsetTool) {};
+
+	/** Clean up resources and don't use NodalOffsetTool anymore */
+	virtual void Shutdown() {};
+
+	/** Called every frame */
+	virtual void Tick(float DeltaTime) {};
+
+	/** Callback when viewport is clicked. Returns false if the event was not handled. */
+	virtual bool OnViewportClicked(const FGeometry& MyGeometry, const FPointerEvent& MouseEvent) { return false;  };
+
+	/** Returns the UI of this calibrator. Expected to only be called once */
+	virtual TSharedRef<SWidget> BuildUI() { return SNew(SVerticalBox); };
+
+	/** Returns the most recently calibrated nodal offset transform, with an error metric */
+	virtual bool GetNodalOffset(FNodalPointOffset& OutNodalOffset, float& OutFocus, float& OutZoom, float& OutError, FText& OutErrorMessage) { return false; };
+
+	/** Returns a descriptive name/title of this nodal offset algorithm */
+	virtual FName FriendlyName() const { return TEXT("Invalid Name"); };
+
+	/** Called when the current offset was saved */
+	virtual void OnSavedNodalOffset() { };
+};
diff --git a/Engine/Plugins/VirtualProduction/CameraCalibration/Source/CameraCalibration/Public/LensFile.h b/Engine/Plugins/VirtualProduction/CameraCalibration/Source/CameraCalibration/Public/LensFile.h
--- a/Engine/Plugins/VirtualProduction/CameraCalibration/Source/CameraCalibration/Public/LensFile.h
+++ b/Engine/Plugins/VirtualProduction/CameraCalibration/Source/CameraCalibration/Public/LensFile.h
@@ -8,6 +8,7 @@
 #include "Engine/Texture.h"
 #include "ICalibratedMapProcessor.h"
 #include "LensData.h"
+#include "Misc/Optional.h"
 #include "Templates/UniquePtr.h"
 #include "Tickable.h"
 #include "UObject/Object.h"
@@ -134,7 +135,6 @@ private:
 };
 
 
-
 /**
 * A data point associating focus and zoom to precalibrated STMap
 */
@@ -265,16 +265,16 @@ public:
 	//~ End FTickableGameObject
 	
 	/** Returns interpolated distortion parameters based on input focus and zoom */
-	bool EvaluateDistortionParameters(float InFocus, float InZoom, FDistortionInfo& OutEvaluatedValue);
+	bool EvaluateDistortionParameters(float InFocus, float InZoom, FDistortionInfo& OutEvaluatedValue) const;
 
 	/** Returns interpolated intrinsic parameters based on input focus and zoom */
-	bool EvaluateIntrinsicParameters(float InFocus, float InZoom, FIntrinsicParameters& OutEvaluatedValue);
+	bool EvaluateIntrinsicParameters(float InFocus, float InZoom, FIntrinsicParameters& OutEvaluatedValue) const;
 
 	/** Draws the distortion map based on evaluation point*/
-	bool EvaluateDistortionData(float InFocus, float InZoom, FVector2D InFilmback, ULensDistortionModelHandlerBase* InLensHandler, FDistortionData& OutDistortionData);
+	bool EvaluateDistortionData(float InFocus, float InZoom, FVector2D InFilmback, ULensDistortionModelHandlerBase* InLensHandler, FDistortionData& OutDistortionData) const;
 
 	/** Returns interpolated nodal point offset based on input focus and zoom */
-	bool EvaluateNodalPointOffset(float InFocus, float InZoom, FNodalPointOffset& OutEvaluatedValue);
+	bool EvaluateNodalPointOffset(float InFocus, float InZoom, FNodalPointOffset& OutEvaluatedValue) const;
 
 	/** Whether focus encoder mapping is configured */
 	bool HasFocusEncoderMapping() const;
@@ -312,10 +312,10 @@ protected:
 	void SetupNoDistortionOutput(ULensDistortionModelHandlerBase* LensHandler, FDistortionData& OutDistortionData) const;
 
 	/** Evaluates distortion based on InFocus and InZoom using parameters */
-	bool EvaluateDistortionForParameters(float InFocus, float InZoom, FVector2D InFilmback, ULensDistortionModelHandlerBase* LensHandler, FDistortionData& OutDistortionData);
+	bool EvaluateDistortionForParameters(float InFocus, float InZoom, FVector2D InFilmback, ULensDistortionModelHandlerBase* LensHandler, FDistortionData& OutDistortionData) const;
 	
 	/** Evaluates distortion based on InFocus and InZoom using STMaps */
-	bool EvaluteDistortionForSTMaps(float InFocus, float InZoom, FVector2D InFilmback, ULensDistortionModelHandlerBase* LensHandler, FDistortionData& OutDistortionData);
+	bool EvaluteDistortionForSTMaps(float InFocus, float InZoom, FVector2D InFilmback, ULensDistortionModelHandlerBase* LensHandler, FDistortionData& OutDistortionData) const;
 	
 public:
 
@@ -395,3 +395,63 @@ public:
 	UPROPERTY(EditAnywhere, BlueprintReadWrite, Category = "Lens File", Meta = (EditCondition = "bOverrideDefaultLensFile"))
 	ULensFile* LensFile = nullptr;
 };
+
+/** Structure that caches the inputs (and other useful bits) used when evaluating the Lens File */
+struct FLensFileEvalData
+{
+	FLensFileEvalData()
+	{
+		Invalidate();
+	};
+
+	ULensFile* LensFile;
+
+	/** The values that should be used as inputs to the Lut in the LensFile */
+	struct
+	{
+		/** Focus input */
+		TOptional<float> Focus;
+
+		/** Iris input */
+		TOptional<float> Iris;
+
+		/** Zoom input */
+		TOptional<float> Zoom;
+	} Input;
+
+	/** Information about the camera associated with the lens evaluation */
+	struct
+	{
+		uint32 UniqueId;
+	} Camera;
+
+	/** Information about the Distortion evaluation */
+	struct
+	{
+		/** True if distotion was applied (and the lens distortion handler updated its state) */
+		bool bWasEvaluated;
+	} Distortion;
+
+	/** Information about the nodal offset evaluation */
+	struct
+	{
+		/** True if the evaluated nodal offset was applied to the camera */
+		bool bWasApplied;
+	} NodalOffset;
+
+	/** Invalidates the data in this structure and avoid using stale or invalid values */
+	void Invalidate()
+	{
+		LensFile = nullptr;
+
+		Input.Focus.Reset();
+		Input.Iris.Reset();
+		Input.Zoom.Reset();
+
+		Camera.UniqueId = INDEX_NONE;
+		Distortion.bWasEvaluated = false;
+		NodalOffset.bWasApplied = false;
+	}
+};
+
+
diff --git a/Engine/Plugins/VirtualProduction/CameraCalibration/Source/CameraCalibrationEditor/CameraCalibrationEditor.Build.cs b/Engine/Plugins/VirtualProduction/CameraCalibration/Source/CameraCalibrationEditor/CameraCalibrationEditor.Build.cs
--- a/Engine/Plugins/VirtualProduction/CameraCalibration/Source/CameraCalibrationEditor/CameraCalibrationEditor.Build.cs
+++ b/Engine/Plugins/VirtualProduction/CameraCalibration/Source/CameraCalibrationEditor/CameraCalibrationEditor.Build.cs
@@ -6,20 +6,36 @@ namespace UnrealBuildTool.Rules
 	{
 		public CameraCalibrationEditor(ReadOnlyTargetRules Target) : base(Target)
 		{
+			if (Target.Configuration == UnrealTargetConfiguration.Debug)
+			{
+				PCHUsage = ModuleRules.PCHUsageMode.UseExplicitOrSharedPCHs;
+				OptimizeCode = CodeOptimization.Never;
+				bUseUnity = false;
+				PCHUsage = PCHUsageMode.NoPCHs;
+			}
+
 			PrivateDependencyModuleNames.AddRange(
 				new string[]
 				{
 					"AssetRegistry",
 					"AssetTools",
 					"CameraCalibration",
+					"CinematicCamera",
+					"Composure",
+					"ComposureLayersEditor",
 					"Core",
 					"CoreUObject",
 	                "CurveEditor",
 					"EditorStyle",
 					"Engine",
 					"InputCore",
+					"LiveLinkCamera",
+					"LiveLinkComponents",
 					"LiveLinkEditor",
 					"LiveLinkInterface",
+					"MediaAssets",
+					"MediaFrameworkUtilities",
+					"OpenCV",
 					"PropertyEditor",
 					"Slate",
 					"SlateCore",
diff --git a/Engine/Plugins/VirtualProduction/CameraCalibration/Source/CameraCalibrationEditor/Private/AssetEditor/CameraCalibrationStepsController.cpp b/Engine/Plugins/VirtualProduction/CameraCalibration/Source/CameraCalibrationEditor/Private/AssetEditor/CameraCalibrationStepsController.cpp
new file mode 100644
--- /dev/null
+++ b/Engine/Plugins/VirtualProduction/CameraCalibration/Source/CameraCalibrationEditor/Private/AssetEditor/CameraCalibrationStepsController.cpp
@@ -0,0 +1,24 @@
+// Copyright Epic Games, Inc. All Rights Reserved.
+
+#include "CameraCalibrationStepsController.h"
+
+#include "CameraCalibrationStep.h"
+#include "CameraCalibrationSubsystem.h"
+#include "CameraCalibrationToolkit.h"
+#include "Widgets/SWidget.h"
+
+FCameraCalibationStepsController::FCameraCalibationStepsController(TWeakPtr<FCameraCalibrationToolkit> InCameraCalibrationToolkit)
+	: CameraCalibrationToolkit(InCameraCalibrationToolkit)
+{
+	check(CameraCalibrationToolkit.IsValid());
+}
+
+FCameraCalibationStepsController::~FCameraCalibationStepsController()
+{
+
+}
+
+TSharedPtr<SWidget> FCameraCalibationStepsController::BuildUI()
+{
+	return nullptr;
+}
diff --git a/Engine/Plugins/VirtualProduction/CameraCalibration/Source/CameraCalibrationEditor/Private/AssetEditor/CameraCalibrationStepsController.h b/Engine/Plugins/VirtualProduction/CameraCalibration/Source/CameraCalibrationEditor/Private/AssetEditor/CameraCalibrationStepsController.h
new file mode 100644
--- /dev/null
+++ b/Engine/Plugins/VirtualProduction/CameraCalibration/Source/CameraCalibrationEditor/Private/AssetEditor/CameraCalibrationStepsController.h
@@ -0,0 +1,32 @@
+// Copyright Epic Games, Inc. All Rights Reserved.
+
+#pragma once
+
+#include "CoreMinimal.h"
+
+
+class FCameraCalibrationToolkit;
+class SWidget;
+class UCameraCalibrationStep;
+
+/**
+ * Controller for SCameraCalibationStepsWidget, where the calibration steps are hosted in.
+ */
+class FCameraCalibationStepsController
+{
+public:
+
+	FCameraCalibationStepsController(TWeakPtr<FCameraCalibrationToolkit> InCameraCalibrationToolkit);
+	~FCameraCalibationStepsController();
+
+	/** Returns the UI that this object controls */
+	TSharedPtr<SWidget> BuildUI();
+
+private:
+
+	/** Pointer to the camera calibration toolkit */
+	TWeakPtr<FCameraCalibrationToolkit> CameraCalibrationToolkit;
+
+	/** Array of the calibration steps that this controller is managing */
+	TArray<TWeakObjectPtr<UCameraCalibrationStep>> CalibrationSteps;
+};
diff --git a/Engine/Plugins/VirtualProduction/CameraCalibration/Source/CameraCalibrationEditor/Private/AssetEditor/CameraCalibrationToolkit.cpp b/Engine/Plugins/VirtualProduction/CameraCalibration/Source/CameraCalibrationEditor/Private/AssetEditor/CameraCalibrationToolkit.cpp
--- a/Engine/Plugins/VirtualProduction/CameraCalibration/Source/CameraCalibrationEditor/Private/AssetEditor/CameraCalibrationToolkit.cpp
+++ b/Engine/Plugins/VirtualProduction/CameraCalibration/Source/CameraCalibrationEditor/Private/AssetEditor/CameraCalibrationToolkit.cpp
@@ -3,6 +3,8 @@
 #include "CameraCalibrationToolkit.h"
 
 #include "LensFile.h"
+#include "Modules/ModuleManager.h"
+#include "PropertyEditorModule.h"
 #include "SLensEvaluation.h"
 #include "SLensFilePanel.h"
 #include "SNodalOffsetToolPanel.h"
@@ -113,7 +115,7 @@ void FCameraCalibrationToolkit::RegisterTabSpawners(const TSharedRef<class FTabM
 		.SetIcon(FSlateIcon(FEditorStyle::GetStyleSetName(), "LevelEditor.GameSettings.Small"));
 
 	InTabManager->RegisterTabSpawner(CameraCalibrationToolkitUtils::LensEvaluationTabId, FOnSpawnTab::CreateSP(this, &FCameraCalibrationToolkit::HandleSpawnLensEvaluationTab))
-		.SetDisplayName(LOCTEXT("NodalOffsetTab", "Nodal Offset"))
+		.SetDisplayName(LOCTEXT("LensEvaluationTab", "Lens Evaluation"))
 		.SetGroup(WorkspaceMenuCategory.ToSharedRef())
 		.SetIcon(FSlateIcon(FEditorStyle::GetStyleSetName(), "LevelEditor.GameSettings.Small"));
 }
diff --git a/Engine/Plugins/VirtualProduction/CameraCalibration/Source/CameraCalibrationEditor/Private/AssetEditor/NodalOffsetTool.cpp b/Engine/Plugins/VirtualProduction/CameraCalibration/Source/CameraCalibrationEditor/Private/AssetEditor/NodalOffsetTool.cpp
new file mode 100644
--- /dev/null
+++ b/Engine/Plugins/VirtualProduction/CameraCalibration/Source/CameraCalibrationEditor/Private/AssetEditor/NodalOffsetTool.cpp
@@ -0,0 +1,876 @@
+// Copyright Epic Games, Inc. All Rights Reserved.
+
+#include "NodalOffsetTool.h"
+
+#include "CoreMinimal.h"
+
+#include "AssetRegistry/AssetData.h"
+#include "CalibrationPointComponent.h"
+#include "Camera/CameraActor.h"
+#include "CameraCalibrationEditorLog.h"
+#include "CameraCalibrationSubsystem.h"
+#include "CameraCalibrationTypes.h"
+#include "CineCameraActor.h"
+#include "CompElementEditorModule.h"
+#include "Components/SceneCaptureComponent2D.h"
+#include "CompositingCaptureBase.h"
+#include "CompositingElement.h"
+#include "CompositingElements/CompositingElementInputs.h"
+#include "CompositingElements/CompositingElementOutputs.h"
+#include "CompositingElements/CompositingElementPasses.h"
+#include "CompositingElements/CompositingElementTransforms.h"
+#include "Containers/Ticker.h"
+#include "Editor.h"
+#include "Engine/TextureRenderTarget2D.h"
+#include "Engine/UserDefinedEnum.h"
+#include "EngineUtils.h"
+#include "ICompElementManager.h"
+#include "LensFile.h"
+#include "LiveLinkCameraController.h"
+#include "LiveLinkComponentController.h"
+#include "Materials/Material.h"
+#include "Materials/MaterialInstanceDynamic.h"
+#include "MediaPlayer.h"
+#include "MediaTexture.h"
+#include "MediaSource.h"
+#include "Misc/MessageDialog.h"
+#include "Models/SphericalLensModel.h"
+#include "Modules/ModuleManager.h"
+#include "Profile/IMediaProfileManager.h"
+#include "Profile/MediaProfile.h"
+#include "TimeSynchronizableMediaSource.h"
+#include "UObject/UObjectGlobals.h"
+#include "UObject/UnrealType.h"
+
+#define LOCTEXT_NAMESPACE "NodalOffsetTool"
+
+
+namespace NodalOffsetTool
+{
+	/** Returns the ICompElementManager object */
+	ICompElementManager* GetCompElementManager()
+	{
+		static ICompElementEditorModule* EditorModule = FModuleManager::Get().GetModulePtr<ICompElementEditorModule>("ComposureLayersEditor");
+
+		TSharedPtr<ICompElementManager> CompElementManager;
+
+		if (EditorModule && (CompElementManager = EditorModule->GetCompElementManager()).IsValid())
+		{
+			return CompElementManager.Get();
+		}
+		return nullptr;
+	}
+
+	/** Retrieves the latest lens evaluation data from the LiveLink controller in the given camera for the given lens file */
+	const FLensFileEvalData* LensFileEvalDataFromCamera(const ACameraActor* InCamera, const ULensFile* InLensFile)
+	{
+		if (!InCamera || !InLensFile)
+		{
+			return nullptr;
+		}
+
+		TArray<ULiveLinkComponentController*> LLComponentControllers;
+		InCamera->GetComponents<ULiveLinkComponentController>(LLComponentControllers);
+		
+		for (const ULiveLinkComponentController* LLComponentController : LLComponentControllers)
+		{
+			for (const TPair<TSubclassOf<ULiveLinkRole>, ULiveLinkControllerBase*> Pair: LLComponentController->ControllerMap)
+			{
+				if (ULiveLinkCameraController* CameraController = Cast<ULiveLinkCameraController>(Pair.Value))
+				{
+					const FLensFileEvalData* OutLensFileEvalData = &CameraController->GetLensFileEvalDataRef();
+
+					if (OutLensFileEvalData->LensFile == InLensFile)
+					{
+						return OutLensFileEvalData;
+					}
+				}
+			}
+		}
+
+		return nullptr;
+	}
+}
+
+FNodalOffsetTool::FNodalOffsetTool(ULensFile* InLensFile)
+{
+	check(InLensFile);
+	LensFile = TStrongObjectPtr<ULensFile>(InLensFile);
+
+	CreateComp();
+
+	TickerHandle = FTicker::GetCoreTicker().AddTicker(FTickerDelegate::CreateRaw(this, &FNodalOffsetTool::OnTick), 0.0f);
+}
+
+FNodalOffsetTool::~FNodalOffsetTool()
+{
+	if (TickerHandle.IsValid())
+	{
+		FTicker::GetCoreTicker().RemoveTicker(TickerHandle);
+	}
+
+	Cleanup();
+}
+
+bool FNodalOffsetTool::OnTick(float DeltaTime)
+{
+	// Update the lens file eval data
+	LensFileEvalData = NodalOffsetTool::LensFileEvalDataFromCamera(Camera.Get(), LensFile.Get());
+
+	if (NodalOffsetAlgo.IsValid())
+	{
+		NodalOffsetAlgo->Tick(DeltaTime);
+	}
+
+	return true;
+}
+
+TWeakObjectPtr<ACompositingElement> FNodalOffsetTool::AddElement(ACompositingElement* Parent, FString& ClassPath, FString& ElementName) const
+{
+	UClass* ElementClass = StaticLoadClass(ACompositingElement::StaticClass(), nullptr, *ClassPath, nullptr, LOAD_None, nullptr);
+
+	if (!ElementClass)
+	{
+		return nullptr;
+	}
+
+	ICompElementManager* CompElementManager = NodalOffsetTool::GetCompElementManager();
+
+	if (!CompElementManager)
+	{
+		return nullptr;
+	}
+
+	TWeakObjectPtr<ACompositingElement> Element = CompElementManager->CreateElement(*ElementName, ElementClass, nullptr, GetTransientPackage());
+
+	if (!Element.IsValid())
+	{
+		return nullptr;
+	}
+
+	if (Parent)
+	{
+		// Attach layer to parent
+		Parent->AttachAsChildLayer(Element.Get());
+
+		// Place element under Parent to keep things organized.
+		Element->AttachToActor(Parent, FAttachmentTransformRules::SnapToTargetNotIncludingScale);
+	}
+
+	return Element;
+}
+
+
+ACompositingElement* FNodalOffsetTool::FindElement(const FString& Name) const
+{
+	// The motivation of finding them is if the level was saved with these objects.
+	// A side effect is that we'll destroy them when closing the calibrator.
+
+	const EActorIteratorFlags Flags = EActorIteratorFlags::SkipPendingKill;
+
+	for (TActorIterator<AActor> It(GetWorld(), ACompositingElement::StaticClass(), Flags); It; ++It)
+	{
+		if (It->GetName() == Name)
+		{
+			return CastChecked<ACompositingElement>(*It);
+		}
+	}
+
+	return nullptr;
+}
+
+void FNodalOffsetTool::Cleanup()
+{
+	if (MediaPlayer.IsValid())
+	{
+		MediaPlayer->Close();
+		MediaPlayer.Reset();
+	}
+
+	if (Comp.IsValid()) // It may have been destroyed manually by the user.
+	{
+		Comp->Destroy();
+		Comp.Reset();
+	}
+	
+	if (CGLayer.IsValid())
+	{
+		CGLayer->Destroy();
+		CGLayer.Reset();
+	}
+	
+	if (MediaPlate.IsValid())
+	{
+		MediaPlate->Destroy();
+		MediaPlate.Reset();
+	}
+	
+	MaterialPass.Reset();
+	Camera.Reset();
+
+	if (NodalOffsetAlgo.IsValid())
+	{
+		NodalOffsetAlgo->Shutdown();
+		NodalOffsetAlgo->RemoveFromRoot();
+		NodalOffsetAlgo.Reset();
+	}
+}
+
+FString FNodalOffsetTool::NamespacedName(const FString&& Name) const
+{
+	return FString::Printf(TEXT("CameraCalib_%s_%s"), *LensFile->GetName(), *Name);
+}
+
+UWorld* FNodalOffsetTool::GetWorld() const
+{
+	return GEditor? GEditor->GetEditorWorldContext().World() : nullptr;
+}
+
+UTextureRenderTarget2D* FNodalOffsetTool::GetRenderTarget() const
+{
+	return RenderTarget.Get();
+}
+
+void FNodalOffsetTool::CreateComp()
+{
+	// Don't do anything if we already created it.
+	if (Comp.IsValid())
+	{
+		// Some items are exposed in the World Outliner, and they will get cleaned up and re-created
+		// when the configurator is closed and re-opened.
+		 
+		UE_LOG(LogCameraCalibrationEditor, Warning, TEXT("Tried to create Comp that already existed"));
+		return;
+	}
+
+	ICompElementManager* CompElementManager = NodalOffsetTool::GetCompElementManager();
+
+	if (!CompElementManager)
+	{
+		UE_LOG(LogCameraCalibrationEditor, Warning, TEXT("Could not find ICompElementManager"));
+		return;
+	}
+
+	// Create Comp parent
+	{
+		const FString CompName = NamespacedName(TEXT("Comp"));
+
+		Comp = FindElement(CompName);
+
+		if(!Comp.IsValid())
+		{
+			Comp = CompElementManager->CreateElement(*CompName, ACompositingElement::StaticClass(), nullptr, GetTransientPackage());
+		}
+
+		if (!Comp.IsValid())
+		{
+			UE_LOG(LogCameraCalibrationEditor, Warning, TEXT("Failed to create 'Comp'"));
+			Cleanup();
+			return;
+		}
+	}
+
+	// Convenience function to add comp elements
+	auto FindOrAddElementToComp = [&](
+		TWeakObjectPtr<ACompositingElement>& Element, 
+		FString&& ElementClassPath, 
+		FString&& ElementName) -> bool
+	{
+		if (Element.IsValid())
+		{
+			return true;
+		}
+
+		Element = FindElement(ElementName);
+
+		if(!Element.IsValid())
+		{
+			Element = AddElement(Comp.Get(), ElementClassPath, ElementName);
+		}
+
+		return Element.IsValid();
+	};
+
+	if (!FindOrAddElementToComp(
+		CGLayer, 
+		TEXT("/Composure/Blueprints/CompositingElements/BP_CgCaptureCompElement.BP_CgCaptureCompElement_C"), 
+		NamespacedName(TEXT("CG"))))
+	{
+		UE_LOG(LogCameraCalibrationEditor, Warning, TEXT("Failed to create 'CG' Layer"));
+		Cleanup();
+		return;
+	}
+
+	if (!FindOrAddElementToComp(
+		MediaPlate, 
+		TEXT("/Composure/Blueprints/CompositingElements/BP_MediaPlateCompElement.BP_MediaPlateCompElement_C"), 
+		NamespacedName(TEXT("MediaPlate"))))
+	{
+		UE_LOG(LogCameraCalibrationEditor, Warning, TEXT("Failed to create 'MediaPlate'"));
+		Cleanup();
+		return;
+	}
+
+	// This updates the Composure panel view
+	CompElementManager->RefreshElementsList();
+
+	// Disable fog on scene capture component of CGLayer
+	{
+		TArray<USceneCaptureComponent2D*> CaptureComponents;
+		CGLayer->GetComponents<USceneCaptureComponent2D>(CaptureComponents);
+
+		for (USceneCaptureComponent2D* CaptureComponent : CaptureComponents)
+		{
+			// We need to disable Fog via the ShowFlagSettings property instead of directly in ShowFlags,
+			// because it will likely be overwritten since they are often replaced by the archetype defaults
+			// and then updated by the ShowFlagSettings.
+
+			FEngineShowFlagsSetting NewFlagSetting;
+			NewFlagSetting.ShowFlagName = FEngineShowFlags::FindNameByIndex(FEngineShowFlags::EShowFlag::SF_Fog);
+			NewFlagSetting.Enabled = false;
+
+			CaptureComponent->ShowFlagSettings.Add(NewFlagSetting);
+
+			if (FProperty* ShowFlagSettingsProperty = CaptureComponent->GetClass()->FindPropertyByName(
+				GET_MEMBER_NAME_CHECKED(USceneCaptureComponent2D, ShowFlagSettings)))
+			{
+				// This PostEditChange will ensure that ShowFlags is updated.
+				FPropertyChangedEvent PropertyChangedEvent(ShowFlagSettingsProperty);
+				CaptureComponent->PostEditChangeProperty(PropertyChangedEvent);
+			}
+		}
+	}
+
+	// Create media player and media texture and assign it to the media input of the media plate layer.
+	//
+	for (UCompositingElementInput* Input : MediaPlate->GetInputsList())
+	{
+		UMediaTextureCompositingInput* MediaInput = Cast<UMediaTextureCompositingInput>(Input);
+
+		if (!MediaInput)
+		{
+			continue;
+		}
+
+		// Create MediaPlayer
+
+		// Using a strong reference prevents the MediaPlayer from going stale when the level is saved.
+		MediaPlayer = TStrongObjectPtr<UMediaPlayer>(NewObject<UMediaPlayer>(
+			GetTransientPackage(), 
+			MakeUniqueObjectName(GetTransientPackage(), UMediaPlayer::StaticClass())
+		));
+
+		if (!MediaPlayer.IsValid())
+		{
+			UE_LOG(LogCameraCalibrationEditor, Warning, TEXT("Failed to create MediaPlayer"));
+			Cleanup();
+			return;
+		}
+
+		MediaPlayer->PlayOnOpen = false;
+
+		// Create MediaTexture
+
+		MediaTexture = NewObject<UMediaTexture>(GetTransientPackage(), NAME_None, RF_Transient);
+
+		if (MediaTexture.IsValid())
+		{
+			MediaTexture->AutoClear = true;
+			MediaTexture->SetMediaPlayer(MediaPlayer.Get());
+			MediaTexture->UpdateResource();
+
+			MediaInput->MediaSource = MediaTexture.Get();
+		}
+
+		// Play the first time-synchronizable media source.
+		if (const UMediaProfile* MediaProfile = IMediaProfileManager::Get().GetCurrentMediaProfile())
+		{
+			for (int32 MediaSourceIdx = 0; MediaSourceIdx < MediaProfile->NumMediaSources(); ++MediaSourceIdx)
+			{
+				if (UTimeSynchronizableMediaSource* TSMediaSource = Cast<UTimeSynchronizableMediaSource>(
+					MediaProfile->GetMediaSource(MediaSourceIdx)))
+				{
+					MediaPlayer->OpenSource(TSMediaSource);
+					MediaPlayer->Play();
+
+					// Break since we don't need to look for more MediaSources.
+					break;
+				}
+			}
+		}
+
+		// Break since don't need to look at more UMediaTextureCompositingInputs.
+		break;
+	}
+
+	// Create material pass that blends MediaPlate with CG.
+	//@todo Make this material selectable by the user.
+
+	MaterialPass = CastChecked<UCompositingElementMaterialPass>(
+		Comp->CreateNewTransformPass(TEXT("CGOverMedia"), UCompositingElementMaterialPass::StaticClass())
+	);
+
+	if (!MaterialPass.IsValid())
+	{
+		UE_LOG(LogCameraCalibrationEditor, Warning, TEXT("Failed to create 'CGOverMedia' UCompositingElementMaterialPass"));
+		Cleanup();
+		return;
+	}
+
+	// Create Material
+	const FString MaterialPath = TEXT("/CameraCalibration/Materials/M_SimulcamCalib.M_SimulcamCalib");
+	UMaterial* Material = Cast<UMaterial>(FSoftObjectPath(MaterialPath).TryLoad());
+
+	if (!Material)
+	{
+		UE_LOG(LogCameraCalibrationEditor, Warning, TEXT("Failed to load %s"), *MaterialPath);
+		Cleanup();
+		return;
+	}
+
+	MaterialPass->SetMaterialInterface(Material);
+
+	// ActorLabel should coincide with ElementName. Name may not.
+	MaterialPass->SetParameterMapping(TEXT("CG"), *CGLayer->GetActorLabel());
+	MaterialPass->SetParameterMapping(TEXT("MediaPlate"), *MediaPlate->GetActorLabel());
+
+	URenderTargetCompositingOutput* RTOutput = Cast<URenderTargetCompositingOutput>(Comp->CreateNewOutputPass(
+		TEXT("SimulcamCalOutput"), 
+		URenderTargetCompositingOutput::StaticClass())
+	);
+
+	if (!RTOutput)
+	{
+		UE_LOG(LogCameraCalibrationEditor, Warning, TEXT("Failed to create URenderTargetCompositingOutput"));
+		Cleanup();
+		return;
+	}
+
+	// Create RenderTarget
+	RenderTarget = NewObject<UTextureRenderTarget2D>(
+		GetTransientPackage(),
+		MakeUniqueObjectName(GetTransientPackage(), UTextureRenderTarget2D::StaticClass())
+	);
+
+	if (!RenderTarget.IsValid())
+	{
+		UE_LOG(LogCameraCalibrationEditor, Warning, TEXT("Failed to create UTextureRenderTarget2D"));
+		Cleanup();
+		return;
+	}
+
+	RenderTarget->RenderTargetFormat = RTF_RGBA8;
+	RenderTarget->ClearColor = FLinearColor::Black;
+	RenderTarget->bAutoGenerateMips = true;
+	RenderTarget->InitAutoFormat(1920, 1080);
+	RenderTarget->UpdateResourceImmediate(true);
+
+	// Assign the RT to the compositing output
+	RTOutput->RenderTarget = RenderTarget.Get();
+
+	// By default, use a camera with our lens. If not found, let composure pick one.
+	if (ACameraActor* CameraGuess = FindFirstCameraWithCurrentLens())
+	{
+		SetCamera(CameraGuess);
+	}
+	else
+	{
+		SetCamera(Comp->FindTargetCamera());
+	}
+}
+
+float FNodalOffsetTool::GetWiperWeight() const
+{
+	float Weight = 0.5f;
+
+	if (!MaterialPass.IsValid())
+	{
+		return Weight;
+	}
+
+	MaterialPass->Material.GetScalarOverride(TEXT("WiperWeight"), Weight);
+	return Weight;
+}
+
+void FNodalOffsetTool::SetWiperWeight(float InWeight)
+{
+	if (MaterialPass.IsValid())
+	{
+		MaterialPass->Material.SetScalarOverride(TEXT("WiperWeight"), InWeight);
+	}
+}
+
+void FNodalOffsetTool::SetCamera(ACameraActor* InCamera)
+{
+	Camera = InCamera;
+
+	// Update the Comp with this camera
+	if (Comp.IsValid())
+	{
+		Comp->SetTargetCamera(InCamera);
+		EnableDistortionInCG();
+	}
+}
+
+void FNodalOffsetTool::EnableDistortionInCG()
+{
+	if (!Comp.IsValid())
+	{
+		return;
+	}
+
+	UCameraCalibrationSubsystem* SubSystem = GEngine->GetEngineSubsystem<UCameraCalibrationSubsystem>();
+
+	if (!SubSystem)
+	{
+		UE_LOG(LogCameraCalibrationEditor, Error, TEXT("Could not find UCameraCalibrationSubsystem"));
+		return;
+	}
+
+	ACineCameraActor* CineCamera = Cast<ACineCameraActor>(Camera.Get());
+
+	if (!CineCamera)
+	{
+		UE_LOG(LogCameraCalibrationEditor, Warning, TEXT("No cine camera selected when trying to enable distortion."));
+		return;
+	}
+
+	// Pick first valid handler that the subsystem finds in our	camera
+	//
+
+	ULensDistortionModelHandlerBase* DistortionHandler = nullptr;
+
+	for (ULensDistortionModelHandlerBase* Handler : SubSystem->GetDistortionModelHandlers(CineCamera->GetCineCameraComponent()))
+	{
+		if (!Handler)
+		{
+			continue;
+		}
+
+		DistortionHandler = Handler;
+
+		break;
+	}
+
+	if (!DistortionHandler)
+	{
+		UE_LOG(LogCameraCalibrationEditor, Warning, TEXT("Could not find a distortion handler in the selected camera"));
+		return;
+	}
+
+	for (ACompositingElement* Element : Comp->GetChildElements())
+	{
+		ACompositingCaptureBase* CaptureBase = Cast<ACompositingCaptureBase>(Element);
+
+		if (!CaptureBase)
+		{
+			continue;
+		}
+
+		CaptureBase->SetApplyDistortion(true);
+		CaptureBase->SetDistortionHandler(DistortionHandler);
+	}
+}
+
+ACameraActor* FNodalOffsetTool::GetCamera() const
+{
+	return Camera.Get();
+}
+
+void FNodalOffsetTool::OnSimulcamViewportClicked(const FGeometry& MyGeometry, const FPointerEvent& MouseEvent)
+{
+	if (!NodalOffsetAlgo.IsValid())
+	{
+		return;
+	}
+
+	// If the algorithm handled the event, we're done
+	if (NodalOffsetAlgo->OnViewportClicked(MyGeometry, MouseEvent))
+	{
+		return;
+	}
+
+	// Toggle video pause with right click.
+	if (MouseEvent.GetEffectingButton() == EKeys::RightMouseButton)
+	{
+		TogglePlay();
+		return;
+	}
+}
+
+ULiveLinkCameraController* FNodalOffsetTool::FindLiveLinkCameraControllerWithLens(const ACameraActor* InCamera, const ULensFile* InLensFile) const
+{
+	if (!InCamera || !InLensFile)
+	{
+		return nullptr;
+	}
+
+	FAssetData LensFileAssetData(InLensFile);
+
+	TArray<ULiveLinkComponentController*> LiveLinkComponents;
+	InCamera->GetComponents<ULiveLinkComponentController>(LiveLinkComponents);
+
+	for (const ULiveLinkComponentController* LiveLinkComponent : LiveLinkComponents)
+	{
+		for (auto It = LiveLinkComponent->ControllerMap.CreateConstIterator(); It; ++It)
+		{
+			ULiveLinkCameraController* CameraController = Cast<ULiveLinkCameraController>(It->Value);
+
+			if (!CameraController)
+			{
+				continue;
+			}
+
+			const ULensFile* CameraLensFile = CameraController->LensFilePicker.GetLensFile();
+
+			if (!CameraLensFile)
+			{
+				continue;
+			}
+
+			if (FAssetData(CameraLensFile, true).PackageName == LensFileAssetData.PackageName)
+			{
+				return CameraController;
+			}
+		}
+	}
+
+	return nullptr;
+}
+
+ACameraActor* FNodalOffsetTool::FindFirstCameraWithCurrentLens() const
+{
+	// We iterate over all cameras in the scene and try to find one that is using the current LensFile
+
+	for (TActorIterator<ACameraActor> CameraItr(GetWorld()); CameraItr; ++CameraItr)
+	{
+		ACameraActor* CameraActor = *CameraItr;
+
+		if (FindLiveLinkCameraControllerWithLens(CameraActor, LensFile.Get()))
+		{
+			return CameraActor;
+		}
+	}
+
+	return nullptr;
+}
+
+void FNodalOffsetTool::TogglePlay()
+{
+	if (!MediaPlayer.IsValid())
+	{
+		return;
+	}
+
+	//@todo Eventually pause should cache the texture instead of relying on player play/pause support.
+
+	if (IsPaused())
+	{
+		MediaPlayer->Play();
+	}
+	else
+	{
+		MediaPlayer->Pause();
+	}
+}
+
+void FNodalOffsetTool::SetNodalOffsetAlgo(const FName& AlgoName)
+{
+	// Ask subsystem for the selected nodal offset algo class
+
+	UCameraCalibrationSubsystem* Subsystem = GEngine->GetEngineSubsystem<UCameraCalibrationSubsystem>();
+	check(Subsystem);
+
+	TSubclassOf<UCameraNodalOffsetAlgo> AlgoClass = Subsystem->GetCameraNodalOffsetAlgo(AlgoName);
+
+	// If it is the same as the existing one, do nothing.
+	if (!NodalOffsetAlgo.IsValid() && !AlgoClass)
+	{
+		return;
+	}
+	else if (NodalOffsetAlgo.IsValid() && (NodalOffsetAlgo->GetClass() == AlgoClass))
+	{
+		return;
+	}
+
+	// Remove old Algo
+	if (NodalOffsetAlgo.IsValid())
+	{
+		NodalOffsetAlgo->Shutdown();
+		NodalOffsetAlgo->RemoveFromRoot();
+		NodalOffsetAlgo.Reset();
+	}
+
+	// If AlgoClass is none, we're done here.
+	if (!AlgoClass)
+	{
+		return;
+	}
+
+	// Create new algo
+	NodalOffsetAlgo = NewObject<UCameraNodalOffsetAlgo>(
+		GetTransientPackage(),
+		AlgoClass,
+		MakeUniqueObjectName(GetTransientPackage(), AlgoClass));
+
+	if (NodalOffsetAlgo.IsValid())
+	{
+		NodalOffsetAlgo->Initialize(SharedThis(this));
+		NodalOffsetAlgo->AddToRoot();
+	}
+}
+
+UCameraNodalOffsetAlgo* FNodalOffsetTool::GetNodalOffsetAlgo() const
+{
+	if (!NodalOffsetAlgo.IsValid())
+	{
+		return nullptr;
+	}
+
+	return NodalOffsetAlgo.Get();
+}
+
+bool FNodalOffsetTool::IsPaused() const
+{
+	if (MediaPlayer.IsValid())
+	{
+		return !MediaPlayer->IsPlaying();
+	}
+
+	return true;
+}
+
+const FLensFileEvalData* FNodalOffsetTool::GetLensFileEvalData() const
+{
+	return LensFileEvalData;
+}
+
+const ULensFile* FNodalOffsetTool::GetLensFile()
+{
+	if (LensFile.IsValid())
+	{
+		return LensFile.Get();
+	}
+
+	return nullptr;
+}
+
+const ULensDistortionModelHandlerBase* FNodalOffsetTool::GetDistortionHandler() const
+{
+	if (!CGLayer.IsValid())
+	{
+		return nullptr;
+	}
+
+	ACompositingCaptureBase* CaptureBase = Cast<ACompositingCaptureBase>(CGLayer.Get());
+
+	if (!CaptureBase)
+	{
+		return nullptr;
+	}
+
+	return CaptureBase->GetDistortionHandler();
+}
+
+void FNodalOffsetTool::OnSaveCurrentNodalOffset()
+{
+	UCameraNodalOffsetAlgo* Algo = GetNodalOffsetAlgo();
+
+	if (!Algo)
+	{
+		FText ErrorMessage = LOCTEXT("NoAlgoFound", "No algo found");
+		FMessageDialog::Open(EAppMsgType::Ok, ErrorMessage);
+		return;
+	}
+
+	FText ErrorMessage;
+
+	FNodalOffsetMapPoint NodalOffsetMapPoint;
+	float ReprojectionError;
+
+	if (!Algo->GetNodalOffset(NodalOffsetMapPoint.NodalOffset, NodalOffsetMapPoint.Focus, NodalOffsetMapPoint.Zoom, ReprojectionError, ErrorMessage))
+	{
+		FMessageDialog::Open(EAppMsgType::Ok, ErrorMessage);
+		return;
+	}
+
+	LensFile->NodalOffsetMapping.Add(NodalOffsetMapPoint);
+
+	// Force bApplyNodalOffset in the LiveLinkCameraController so that we can see the effect right away
+	if (ULiveLinkCameraController* LiveLinkCameraController = FindLiveLinkCameraControllerWithLens(Camera.Get(), LensFile.Get()))
+	{
+		LiveLinkCameraController->SetApplyNodalOffset(true);
+	}
+
+	Algo->OnSavedNodalOffset();
+}
+
+bool FNodalOffsetTool::SetMediaSourceUrl(const FString& InMediaSourceUrl)
+{
+	const UMediaProfile* MediaProfile = IMediaProfileManager::Get().GetCurrentMediaProfile();
+
+	if (!MediaProfile || !MediaPlayer.IsValid())
+	{
+		return false;
+	}
+
+	// If we're already playing it, we're done
+	if (InMediaSourceUrl == GetMediaSourceUrl())
+	{
+		return true;
+	}
+
+	if (InMediaSourceUrl == TEXT("None") || !InMediaSourceUrl.Len())
+	{
+		MediaPlayer->Close();
+		return true;
+	}
+
+	for (int32 MediaSourceIdx = 0; MediaSourceIdx < MediaProfile->NumMediaSources(); ++MediaSourceIdx)
+	{
+		UTimeSynchronizableMediaSource* TSMediaSource = Cast<UTimeSynchronizableMediaSource>(
+			MediaProfile->GetMediaSource(MediaSourceIdx));
+
+		if (!TSMediaSource || TSMediaSource->GetUrl() != InMediaSourceUrl)
+		{
+			continue;
+		}
+
+		MediaPlayer->OpenSource(TSMediaSource);
+		MediaPlayer->Play();
+
+		return true;
+	}
+
+	return false;
+}
+
+/** Gets the current media source url being played. Empty if None */
+FString FNodalOffsetTool::GetMediaSourceUrl() const
+{
+	if (!MediaPlayer.IsValid())
+	{
+		return TEXT("");
+	}
+
+	return MediaPlayer->GetUrl();
+}
+
+void FNodalOffsetTool::FindMediaSourceUrls(TArray<TSharedPtr<FString>>& OutMediaSourceUrls) const
+{
+	const UMediaProfile* MediaProfile = IMediaProfileManager::Get().GetCurrentMediaProfile();
+
+	if (!MediaProfile)
+	{
+		return;
+	}
+
+	for (int32 MediaSourceIdx = 0; MediaSourceIdx < MediaProfile->NumMediaSources(); ++MediaSourceIdx)
+	{
+		if (UTimeSynchronizableMediaSource* TSMediaSource = Cast<UTimeSynchronizableMediaSource>(
+			MediaProfile->GetMediaSource(MediaSourceIdx)))
+		{
+			OutMediaSourceUrls.Add(MakeShareable(new FString(TSMediaSource->GetUrl())));
+		}
+	}
+}
\ No newline at end of file
diff --git a/Engine/Plugins/VirtualProduction/CameraCalibration/Source/CameraCalibrationEditor/Private/AssetEditor/NodalOffsetTool.h b/Engine/Plugins/VirtualProduction/CameraCalibration/Source/CameraCalibrationEditor/Private/AssetEditor/NodalOffsetTool.h
new file mode 100644
--- /dev/null
+++ b/Engine/Plugins/VirtualProduction/CameraCalibration/Source/CameraCalibrationEditor/Private/AssetEditor/NodalOffsetTool.h
@@ -0,0 +1,160 @@
+// Copyright Epic Games, Inc. All Rights Reserved.
+
+#pragma once
+
+#include "CoreMinimal.h"
+
+#include "Engine/World.h"
+#include "UObject/StrongObjectPtr.h"
+
+class ACameraActor;
+class ACompositingElement;
+class UMediaPlayer;
+class UMediaTexture;
+class UCameraNodalOffsetAlgo;
+class UCompositingElementMaterialPass;
+class ULensDistortionModelHandlerBase;
+class ULensFile;
+class ULiveLinkCameraController;
+class UTextureRenderTarget2D;
+
+struct FGeometry;
+struct FLensFileEvalData;
+struct FPointerEvent;
+
+/**
+ * FNodalOffsetTool is the controller for the nodal offset tool panel.
+ * It has the logic to bridge user input like selection of nodal offset algorithm or CG camera
+ * with the actions that follow. It houses convenience functions used to generate the data
+ * of what is presented to the user, and holds pointers to the relevant objects and structures.
+ */
+class FNodalOffsetTool : public TSharedFromThis<FNodalOffsetTool>
+{
+public:
+
+	FNodalOffsetTool(ULensFile* InLensFile);
+	~FNodalOffsetTool();
+
+	/** Creates composite with CG and selected media source */
+	void CreateComp();
+
+	/** Returns the render target of the Comp */
+	UTextureRenderTarget2D* GetRenderTarget() const;
+
+	/** Returns the CG weight that is composited on top of the media */
+	float GetWiperWeight() const;
+
+	/** Sets the weight/alpha for that CG that is composited on top of the media. 0 means invisible. */
+	void SetWiperWeight(float InWeight);
+
+	/** Sets the camera used for the CG */
+	void SetCamera(ACameraActor* InCamera);
+
+	/** Returns the camera used for the CG */
+	ACameraActor* GetCamera() const;
+
+	/** Toggles the play/stop state of the media player */
+	void TogglePlay();
+
+	/** Selects the nodal offset algorithm by name */
+	void SetNodalOffsetAlgo(const FName& AlgoName);
+
+	/** Returns the currently selected nodal offset algorithm */
+	UCameraNodalOffsetAlgo* GetNodalOffsetAlgo() const;
+
+	/** Returns true if the media player is paused */
+	bool IsPaused() const;
+
+	/** Returns the latest data used when evaluating the lens */
+	const FLensFileEvalData* GetLensFileEvalData() const;
+
+	/** Returns the LensFile that this nodal offset tool is using */
+	const ULensFile* GetLensFile();
+
+	/** Returns the distortion handler used to distort the CG being displayed in the simulcam viewport */
+	const ULensDistortionModelHandlerBase* GetDistortionHandler() const;
+
+	/** Finds the LiveLinkCameraController used in the given CameraActor that is also using the given LensFile */
+	ULiveLinkCameraController* FindLiveLinkCameraControllerWithLens(const ACameraActor* CameraActor, const ULensFile* InLensFile) const;
+
+	/** Sets the media source url to be played. Returns true if the url is a valid media source */
+	bool SetMediaSourceUrl(const FString& InMediaSourceUrl);
+	
+	/** Finds available media sources and adds their urls to the given array */
+	void FindMediaSourceUrls(TArray<TSharedPtr<FString>>& OutMediaSourceUrls) const;
+
+	/** Gets the current media source url being played. Empty if None */
+	FString GetMediaSourceUrl() const;
+
+public:
+
+	/** Called by the UI when the Simulcam Viewport is clicked */
+	void OnSimulcamViewportClicked(const FGeometry& MyGeometry, const FPointerEvent& MouseEvent);
+
+	/** Called by the UI when the user wants to save the nodal offset that the current algorithm is providing */
+	void OnSaveCurrentNodalOffset();
+
+private:
+
+	/** Returns a namespaced version of the given name. Useful to generate names unique to this lens file */
+	FString NamespacedName(const FString&& Name) const;
+
+	/** Finds the world being used by the tool for finding and spawning objects */
+	UWorld* GetWorld() const;
+
+	/** Finds an existing comp element based on its name. */
+	ACompositingElement* FindElement(const FString& Name) const;
+
+	/** Spawns a new compositing element of the given class and parent element */
+	TWeakObjectPtr<ACompositingElement> AddElement(ACompositingElement* Parent, FString& ClassPath, FString& ElementName) const;
+
+	/** Releases resources used by the nodal offset tool, including the transient comp elements */
+	void Cleanup();
+
+	/** Convenience function that returns the first camera it finds that is using the lens associated with this object. */
+	ACameraActor* FindFirstCameraWithCurrentLens() const;
+
+	/** Enables distortion in the CG comp */
+	void EnableDistortionInCG();
+
+	/** Called by the core ticker */
+	bool OnTick(float DeltaTime);
+
+private:
+
+	/** Strong pointer to the LensFile used by this object */
+	TStrongObjectPtr<class ULensFile> LensFile;
+
+	/** The parent comp element */
+	TWeakObjectPtr<ACompositingElement> Comp;
+
+	/** The CG layer comp element */
+	TWeakObjectPtr<ACompositingElement> CGLayer;
+
+	/** The MediaPlate comp element*/
+	TWeakObjectPtr<ACompositingElement> MediaPlate;
+
+	/** The output render target of the comp */
+	TWeakObjectPtr<UTextureRenderTarget2D> RenderTarget;
+
+	/** The media texture used by the media plate */
+	TWeakObjectPtr<UMediaTexture> MediaTexture;
+
+	/** The media player that is playing the selected media source */
+	TStrongObjectPtr<UMediaPlayer> MediaPlayer;
+
+	/** The material pass the does the CG + MediaPlate composite with a wiper weight */
+	TWeakObjectPtr<UCompositingElementMaterialPass> MaterialPass;
+
+	/** The currently selected camera */
+	TWeakObjectPtr<ACameraActor> Camera;
+
+	/** The currently selected nodal offset algorithm */
+	TWeakObjectPtr<UCameraNodalOffsetAlgo> NodalOffsetAlgo;
+
+	/** The delegate for the core ticker callback */
+	FDelegateHandle TickerHandle;
+
+	/** Pointer to the LensFileEvalData used in the current frame. Only valid during the current frame. */
+	const FLensFileEvalData* LensFileEvalData;
+};
diff --git a/Engine/Plugins/VirtualProduction/CameraCalibration/Source/CameraCalibrationEditor/Private/AssetEditor/SNodalOffsetToolPanel.cpp b/Engine/Plugins/VirtualProduction/CameraCalibration/Source/CameraCalibrationEditor/Private/AssetEditor/SNodalOffsetToolPanel.cpp
--- a/Engine/Plugins/VirtualProduction/CameraCalibration/Source/CameraCalibrationEditor/Private/AssetEditor/SNodalOffsetToolPanel.cpp
+++ b/Engine/Plugins/VirtualProduction/CameraCalibration/Source/CameraCalibrationEditor/Private/AssetEditor/SNodalOffsetToolPanel.cpp
@@ -2,8 +2,29 @@
 
 #include "SNodalOffsetToolPanel.h"
 
+#include "AssetRegistry/AssetData.h"
+#include "Camera/CameraActor.h"
+#include "CameraCalibrationSubsystem.h"
+#include "CameraNodalOffsetAlgo.h"
+#include "EditorStyleSet.h"
+#include "Engine/Selection.h"
+#include "Engine/StaticMeshActor.h"
+#include "Engine/TextureRenderTarget2D.h"
+#include "IContentBrowserSingleton.h"
+#include "LensFile.h"
+#include "NodalOffsetTool.h"
+#include "PropertyCustomizationHelpers.h"
+#include "SSimulcamViewport.h"
+#include "Styling/CoreStyle.h"
+#include "Styling/SlateStyle.h"
+#include "UI/CameraCalibrationWidgetHelpers.h"
+#include "Widgets/Images/SImage.h"
+#include "Widgets/Input/SButton.h"
+#include "Widgets/Input/SComboBox.h"
+#include "Widgets/Input/SSpinBox.h"
 #include "Widgets/SBoxPanel.h"
 #include "Widgets/Text/STextBlock.h"
+#include "Widgets/Views/STreeView.h"
 
 
 #define LOCTEXT_NAMESPACE "NodalOffsetTool"
@@ -12,19 +33,302 @@
 void SNodalOffsetToolPanel::Construct(const FArguments& InArgs, ULensFile* InLensFile)
 {
 	LensFile = TStrongObjectPtr<ULensFile>(InLensFile);
+	NodalOffsetTool = MakeShared<FNodalOffsetTool>(InLensFile);
+
+	// This will be the widget wrapper of the custom algo UI.
+	NodalOffsetUI = SNew(SVerticalBox);
 
 	ChildSlot
 	[
-		SNew(SVerticalBox)
-		+ SVerticalBox::Slot()
-		.FillHeight(1.0f)
-		.VAlign(VAlign_Center)
-		[
-			SNew(STextBlock)
-			.Text(LOCTEXT("TEMPPLACEHOLDER", "Nothing to show"))
+		SNew(SHorizontalBox)
+
+		+ SHorizontalBox::Slot() // Simulcam Viewport
+		.FillWidth(0.75f) 
+		[ 
+			SNew(SSimulcamViewport, NodalOffsetTool->GetRenderTarget())
+			.OnSimulcamViewportClicked_Raw(NodalOffsetTool.Get(), &FNodalOffsetTool::OnSimulcamViewportClicked)
+		]
+
+		+ SHorizontalBox::Slot() // Right toolbar
+		.FillWidth(0.25f)
+		[ 
+			SNew(SVerticalBox)
+
+			+ SVerticalBox::Slot() // Wiper
+			.MaxHeight(FCameraCalibrationWidgetHelpers::DefaultRowHeight)
+			[ FCameraCalibrationWidgetHelpers::BuildLabelWidgetPair(LOCTEXT("Transparency", "Transparency"), BuildSimulcamWiperWidget())]
+
+			+ SVerticalBox::Slot() // Camera picker
+			.MaxHeight(FCameraCalibrationWidgetHelpers::DefaultRowHeight)
+			[ FCameraCalibrationWidgetHelpers::BuildLabelWidgetPair(LOCTEXT("Camera", "Camera"), BuildCameraPickerWidget())]
+
+			+ SVerticalBox::Slot() // Media Source picker
+			.MaxHeight(FCameraCalibrationWidgetHelpers::DefaultRowHeight)
+			[FCameraCalibrationWidgetHelpers::BuildLabelWidgetPair(LOCTEXT("MediaSource", "Media Source"), BuildMediaSourceWidget())]
+
+			+ SVerticalBox::Slot() // Algo picker
+			.MaxHeight(FCameraCalibrationWidgetHelpers::DefaultRowHeight)
+			[ FCameraCalibrationWidgetHelpers::BuildLabelWidgetPair(LOCTEXT("NodalOffsetAlgo", "Nodal Offset Algo"), BuildNodalOffsetAlgoPickerWidget())]
+
+			+ SVerticalBox::Slot() // Algo UI
+			.AutoHeight()
+			[ BuildNodalOffsetUIWrapper() ]
+
+			+ SVerticalBox::Slot() // Save Offset
+			.AutoHeight()
+			.Padding(0, 20)
+			[
+				SNew(SButton).Text(LOCTEXT("AddToNodalOffsetLUT", "Add To Nodal Offset LUT"))
+				.HAlign(HAlign_Center)
+				.VAlign(VAlign_Center)
+				.OnClicked_Lambda([&]() -> FReply
+				{
+					if (NodalOffsetTool.IsValid())
+					{
+						NodalOffsetTool->OnSaveCurrentNodalOffset();
+					}
+					return FReply::Handled();
+				})
+			]
 		]
 	];
 }
 
+TSharedRef<SWidget> SNodalOffsetToolPanel::BuildNodalOffsetUIWrapper()
+{
+	const UCameraNodalOffsetAlgo* Algo = NodalOffsetTool->GetNodalOffsetAlgo();
+
+	NodalOffsetUITitle = SNew(STextBlock)
+		.Font(FCoreStyle::GetDefaultFontStyle("Bold", 13))
+		.Text(Algo ? FText::FromName(Algo->FriendlyName()) : LOCTEXT("None", "None"))
+		.Justification(ETextJustify::Center);
+
+	return SNew(SVerticalBox)
+
+		+ SVerticalBox::Slot() // Title
+		.AutoHeight()
+		.Padding(0,10)
+		[ NodalOffsetUITitle.ToSharedRef() ]
+
+		+ SVerticalBox::Slot() // Algo's Widget
+		.AutoHeight()
+		[ NodalOffsetUI.ToSharedRef() ];
+}
+
+void SNodalOffsetToolPanel::UpdateNodalOffsetUI()
+{
+	check(AlgosComboBox.IsValid());
+
+	// Get current algo to later compare with new one
+	const UCameraNodalOffsetAlgo* OldAlgo = NodalOffsetTool->GetNodalOffsetAlgo();
+
+	// Set new algo by name
+	const FName AlgoName(*AlgosComboBox->GetSelectedItem());
+	NodalOffsetTool->SetNodalOffsetAlgo(AlgoName);
+
+	// Get the new algo
+	UCameraNodalOffsetAlgo* Algo = NodalOffsetTool->GetNodalOffsetAlgo();
+
+	// nullptr may indicate that it was unregistered, so refresh combobox options.
+	if (!Algo)
+	{
+		UpdateAlgosOptions();
+	}
+
+	// If we didn't change the algo, we're done here.
+	if (Algo == OldAlgo)
+	{
+		return;
+	}
+
+	// Remove old UI
+	check(NodalOffsetUI.IsValid());
+	NodalOffsetUI->ClearChildren();
+
+	// Assign GUI
+	NodalOffsetUI->AddSlot() [Algo->BuildUI()];
+
+	// Update Title
+	if (NodalOffsetUITitle.IsValid())
+	{
+		NodalOffsetUITitle->SetText(FText::FromName(Algo->FriendlyName()));
+	}
+}
+
+TSharedRef<SWidget> SNodalOffsetToolPanel::BuildCameraPickerWidget()
+{
+	return SNew(SObjectPropertyEntryBox)
+		.AllowedClass(ACameraActor::StaticClass())
+		.OnObjectChanged_Lambda([&](const FAssetData& AssetData)
+		{
+			if (AssetData.IsValid())
+			{
+				NodalOffsetTool->SetCamera(Cast<ACameraActor>(AssetData.GetAsset()));
+			}
+		})
+		.ObjectPath_Lambda([&]() -> FString
+		{
+			if (ACameraActor* Camera = NodalOffsetTool->GetCamera())
+			{
+				FAssetData AssetData(Camera, true);
+				return AssetData.ObjectPath.ToString();
+			}
+
+			return TEXT("");
+		});
+}
+
+TSharedRef<SWidget> SNodalOffsetToolPanel::BuildSimulcamWiperWidget()
+{
+	return SNew(SSpinBox<float>)
+		.Value_Lambda([&]() { return NodalOffsetTool->GetWiperWeight(); })
+		.ToolTipText(LOCTEXT("CGWiper", "CG/Media Wiper"))
+		.OnValueChanged_Lambda([&](double InValue)
+		{
+			NodalOffsetTool->SetWiperWeight(float(InValue));
+		})
+		.MinValue(0.0f)
+		.MaxValue(1.0f)
+		.MinSliderValue(0.0f)
+		.MaxSliderValue(1.0f)
+		.ClearKeyboardFocusOnCommit(true)
+		.Delta(0.1f);
+}
+
+void SNodalOffsetToolPanel::UpdateAlgosOptions()
+{
+	CurrentAlgos.Empty();
+
+	// Ask the subsystem for the list of registered Algos
+
+	UCameraCalibrationSubsystem* Subsystem = GEngine->GetEngineSubsystem<UCameraCalibrationSubsystem>();
+	check(Subsystem);
+
+	for (FName& Name : Subsystem->GetCameraNodalOffsetAlgos())
+	{
+		CurrentAlgos.Add(MakeShareable(new FString(Name.ToString())));
+	}
+
+	// Ask the ComboBox to refresh its options from its source (that we just updated)
+	AlgosComboBox->RefreshOptions();
+}
+
+TSharedRef<SWidget> SNodalOffsetToolPanel::BuildNodalOffsetAlgoPickerWidget()
+{
+	// Create ComboBox widget
+
+	AlgosComboBox = SNew(SComboBox<TSharedPtr<FString>>)
+		.OptionsSource(&CurrentAlgos)
+		.OnSelectionChanged_Lambda([&](TSharedPtr<FString> NewValue, ESelectInfo::Type Type) -> void
+		{
+			// Replace the custom algo widget
+			UpdateNodalOffsetUI();
+		})
+		.OnGenerateWidget_Lambda([&](TSharedPtr<FString> InOption) -> TSharedRef<SWidget>
+		{
+			return SNew(STextBlock).Text(FText::FromString(*InOption));
+		})
+		.InitiallySelectedItem(nullptr)
+		[
+			SNew(STextBlock)
+			.Text_Lambda([&]() -> FText
+			{
+				if (AlgosComboBox.IsValid() && AlgosComboBox->GetSelectedItem().IsValid())
+				{
+					return FText::FromString(*AlgosComboBox->GetSelectedItem());
+				}
+
+				return LOCTEXT("InvalidComboOption", "Invalid");
+			})
+		];
+
+	// Update the object holding this combobox's options source
+	UpdateAlgosOptions();
+
+	// Pick the first option by default (if available)
+	if (CurrentAlgos.Num())
+	{
+		AlgosComboBox->SetSelectedItem(CurrentAlgos[0]);
+	}
+	else
+	{
+		AlgosComboBox->SetSelectedItem(nullptr);
+	}
+
+	return AlgosComboBox.ToSharedRef();
+}
+
+void SNodalOffsetToolPanel::UpdateMediaSourcesOptions()
+{
+	CurrentMediaSources.Empty();
+
+	if (NodalOffsetTool.IsValid())
+	{
+		NodalOffsetTool->FindMediaSourceUrls(CurrentMediaSources);
+	}
+
+	// Add a "None" option
+	CurrentMediaSources.Add(MakeShareable(new FString(TEXT("None"))));
+
+	check(MediaSourcesComboBox.IsValid());
+
+	// Ask the ComboBox to refresh its options from its source (that we just updated)
+	MediaSourcesComboBox->RefreshOptions();
+
+	// Make sure we show the item that is selected
+	const FString MediaSourceUrl = NodalOffsetTool->GetMediaSourceUrl();
+
+	for (const TSharedPtr<FString>& MediaSourceUrlItem: CurrentMediaSources)
+	{
+		check(MediaSourceUrlItem.IsValid());
+
+		if (*MediaSourceUrlItem == MediaSourceUrl)
+		{
+			MediaSourcesComboBox->SetSelectedItem(MediaSourceUrlItem);
+			return;
+		}
+	}
+
+	// If we arrived here, we fall back to "None"
+	MediaSourcesComboBox->SetSelectedItem(CurrentMediaSources[CurrentMediaSources.Num() - 1]);
+
+	return;
+}
+
+TSharedRef<SWidget> SNodalOffsetToolPanel::BuildMediaSourceWidget()
+{
+	MediaSourcesComboBox = SNew(SComboBox<TSharedPtr<FString>>)
+		.OptionsSource(&CurrentMediaSources)
+		.OnSelectionChanged_Lambda([&](TSharedPtr<FString> NewValue, ESelectInfo::Type Type) -> void
+		{
+			if (!NodalOffsetTool.IsValid() || !NewValue.IsValid())
+			{
+				return;
+			}
+
+			NodalOffsetTool->SetMediaSourceUrl(*NewValue);
+		})
+		.OnGenerateWidget_Lambda([&](TSharedPtr<FString> InOption) -> TSharedRef<SWidget>
+		{
+			return SNew(STextBlock).Text(FText::FromString(*InOption));
+		})
+		.InitiallySelectedItem(nullptr)
+		[
+			SNew(STextBlock)
+			.Text_Lambda([&]() -> FText
+			{
+				if (MediaSourcesComboBox.IsValid() && MediaSourcesComboBox->GetSelectedItem().IsValid())
+				{
+					return FText::FromString(*MediaSourcesComboBox->GetSelectedItem());
+				}
+
+				return LOCTEXT("InvalidComboOption", "Invalid");
+			})
+		];
+
+	UpdateMediaSourcesOptions();
+
+	return MediaSourcesComboBox.ToSharedRef();
+}
 
-#undef LOCTEXT_NAMESPACE /*NodalOffsetTool*/
+#undef LOCTEXT_NAMESPACE
diff --git a/Engine/Plugins/VirtualProduction/CameraCalibration/Source/CameraCalibrationEditor/Private/AssetEditor/SNodalOffsetToolPanel.h b/Engine/Plugins/VirtualProduction/CameraCalibration/Source/CameraCalibrationEditor/Private/AssetEditor/SNodalOffsetToolPanel.h
--- a/Engine/Plugins/VirtualProduction/CameraCalibration/Source/CameraCalibrationEditor/Private/AssetEditor/SNodalOffsetToolPanel.h
+++ b/Engine/Plugins/VirtualProduction/CameraCalibration/Source/CameraCalibrationEditor/Private/AssetEditor/SNodalOffsetToolPanel.h
@@ -1,14 +1,22 @@
 // Copyright Epic Games, Inc. All Rights Reserved.
 #pragma once
 
+#include "Templates/SharedPointer.h"
+#include "Widgets/DeclarativeSyntaxSupport.h"
 #include "Widgets/SCompoundWidget.h"
-
-#include "LensFile.h"
 #include "UObject/StrongObjectPtr.h"
 
+class FString;
+class ULensFile;
+
+template<typename OptionType>
+class SComboBox;
+
+struct FArguments;
 
 /**
- * Placeholder for camera calibration tools tab
+ * UI for the nodal offset calibration.
+ * It also holds the UI given by the selected nodal offset algorithm.
  */
 class SNodalOffsetToolPanel : public SCompoundWidget
 {
@@ -18,9 +26,57 @@ class SNodalOffsetToolPanel : public SCompoundWidget
 public:
 	void Construct(const FArguments& InArgs, ULensFile* InLensFile);
 
+private:
+
+	/** Builds the UI used to pick the camera used for the CG layer of the comp */
+	TSharedRef<SWidget> BuildCameraPickerWidget();
+
+	/** Builds the UI for the simulcam wiper */
+	TSharedRef<SWidget> BuildSimulcamWiperWidget();
+
+	/** Builds the wrapper for the currently selected nodal offset UI */
+	TSharedRef<SWidget> BuildNodalOffsetUIWrapper();
+
+	/** Builds the UI for the nodal offset algorithm picker */
+	TSharedRef<SWidget> BuildNodalOffsetAlgoPickerWidget();
+
+	/** Builds the UI for the media source picker */
+	TSharedRef<SWidget> BuildMediaSourceWidget();
+
+	/** Updates the UI so that it matches the selected nodal offset algorithm (if necessary) */
+	void UpdateNodalOffsetUI();
 
+	/** Refreshes the list of available nodal offset algorithms shown in the AlgosComboBox */
+	void UpdateAlgosOptions();
+
+	/** Refreshes the list of available media sources shown in the MediaSourcesComboBox */
+	void UpdateMediaSourcesOptions();
+
+private:
+
+	/** The lens asset */
+	TStrongObjectPtr<class ULensFile> LensFile;
+
+	/** The nodal offset tool controller object */
+	TSharedPtr<class FNodalOffsetTool> NodalOffsetTool;
+
+	/** The box containing the UI given by the selected nodal offset algorithm */
+	TSharedPtr<class SVerticalBox> NodalOffsetUI;
+
+	/** Displays the title of the nodal offset UI */
+	TSharedPtr<class STextBlock> NodalOffsetUITitle;
+	
 private:
 
-	/** Holds the preset asset. */
-	TStrongObjectPtr<ULensFile> LensFile;
+	/** Options source for the AlgosComboBox. Lists the currently available nodal offset algorithms */
+	TArray<TSharedPtr<FString>> CurrentAlgos;
+
+	/** The combobox that presents the available nodal offset algorithms */
+	TSharedPtr<SComboBox<TSharedPtr<FString>>> AlgosComboBox;
+
+	/** Options source for the MediaSourcesComboBox. Lists the currently available media sources */
+	TArray<TSharedPtr<FString>> CurrentMediaSources;
+
+	/** The combobox that presents the available media sources */
+	TSharedPtr<SComboBox<TSharedPtr<FString>>> MediaSourcesComboBox;
 };
diff --git a/Engine/Plugins/VirtualProduction/CameraCalibration/Source/CameraCalibrationEditor/Private/AssetEditor/SSimulcamViewport.cpp b/Engine/Plugins/VirtualProduction/CameraCalibration/Source/CameraCalibrationEditor/Private/AssetEditor/SSimulcamViewport.cpp
new file mode 100644
--- /dev/null
+++ b/Engine/Plugins/VirtualProduction/CameraCalibration/Source/CameraCalibrationEditor/Private/AssetEditor/SSimulcamViewport.cpp
@@ -0,0 +1,102 @@
+// Copyright Epic Games, Inc. All Rights Reserved.
+
+#include "SSimulcamViewport.h"
+
+#include "EditorStyleSet.h"
+#include "Materials/Material.h"
+#include "NodalOffsetTool.h"
+#include "Styling/SlateBrush.h"
+#include "Widgets/Layout/SScaleBox.h"
+#include "Widgets/Images/SImage.h"
+#include "Widgets/Input/SButton.h"
+
+
+SSimulcamViewport::SSimulcamViewport()
+	: Collector(this)
+	, Material(nullptr)
+	, MaterialBrush(nullptr)
+	, TextureSampler(nullptr)
+{ }
+
+void SSimulcamViewport::Construct(const FArguments& InArgs, UTexture* InTexture)
+{
+	OnSimulcamViewportClicked = InArgs._OnSimulcamViewportClicked;
+
+	// The Slate brush that renders the material.
+	BrushImageSize = InArgs._BrushImageSize;
+
+	if (InTexture != nullptr)
+	{
+		// create wrapper material
+		Material = NewObject<UMaterial>(GetTransientPackage(), NAME_None, RF_Transient);
+
+		if (Material != nullptr)
+		{
+			TextureSampler = NewObject<UMaterialExpressionTextureSample>(Material);
+			{
+				TextureSampler->Texture = InTexture;
+				TextureSampler->AutoSetSampleType();
+			}
+
+			FExpressionOutput& Output = TextureSampler->GetOutputs()[0];
+			FExpressionInput& Input = Material->EmissiveColor;
+			{
+				Input.Expression = TextureSampler;
+				Input.Mask = Output.Mask;
+				Input.MaskR = Output.MaskR;
+				Input.MaskG = Output.MaskG;
+				Input.MaskB = Output.MaskB;
+				Input.MaskA = Output.MaskA;
+			}
+
+			FExpressionInput& Opacity = Material->Opacity;
+			{
+				Opacity.Expression = TextureSampler;
+				Opacity.Mask = Output.Mask;
+				Opacity.MaskR = 0;
+				Opacity.MaskG = 0;
+				Opacity.MaskB = 0;
+				Opacity.MaskA = 1;
+			}
+
+			Material->BlendMode = BLEND_AlphaComposite;
+
+			Material->Expressions.Add(TextureSampler);
+			Material->MaterialDomain = EMaterialDomain::MD_UI;
+			Material->PostEditChange();
+		}
+
+		// create Slate brush
+		MaterialBrush = MakeShared<FSlateBrush>();
+		{
+			MaterialBrush->SetResourceObject(Material);
+		}
+	}
+
+	ChildSlot
+	[
+		SNew(SScaleBox)
+		.Stretch_Lambda([]() -> EStretch::Type { return EStretch::Fill;	})
+		[
+			SNew(SImage)
+			.Image(MaterialBrush.IsValid() ? MaterialBrush.Get() : FEditorStyle::GetBrush("WhiteTexture"))
+			.OnMouseButtonDown_Lambda([&](const FGeometry& MyGeometry, const FPointerEvent& MouseEvent)
+			{
+				OnSimulcamViewportClicked.ExecuteIfBound(MyGeometry, MouseEvent);
+				return FReply::Handled();
+			})
+		]
+	];
+}
+
+void SSimulcamViewport::Tick(const FGeometry& InAllottedGeometry, const double InCurrentTime, const float InDeltaTime)
+{
+	SCompoundWidget::Tick(InAllottedGeometry, InCurrentTime, InDeltaTime);
+
+	if (BrushImageSize.IsSet())
+	{
+		const FVector2D Size = BrushImageSize.Get();
+		MaterialBrush->ImageSize.X = Size.X;
+		MaterialBrush->ImageSize.Y = Size.Y;
+	}
+}
diff --git a/Engine/Plugins/VirtualProduction/CameraCalibration/Source/CameraCalibrationEditor/Private/AssetEditor/SSimulcamViewport.h b/Engine/Plugins/VirtualProduction/CameraCalibration/Source/CameraCalibrationEditor/Private/AssetEditor/SSimulcamViewport.h
new file mode 100644
--- /dev/null
+++ b/Engine/Plugins/VirtualProduction/CameraCalibration/Source/CameraCalibrationEditor/Private/AssetEditor/SSimulcamViewport.h
@@ -0,0 +1,98 @@
+// Copyright Epic Games, Inc. All Rights Reserved.
+
+#pragma once
+
+#include "Widgets/DeclarativeSyntaxSupport.h"
+#include "Widgets/SCompoundWidget.h"
+#include "UObject/GCObject.h"
+
+
+class UMaterial;
+class UMaterialExpressionTextureSample;
+class UTexture;
+
+struct FGeometry;
+struct FPointerEvent;
+
+/** Delegate to be executed when the viewport area is clicked */
+DECLARE_DELEGATE_TwoParams(FSimulcamViewportClickedEventHandler, const FGeometry& MyGeometry, const FPointerEvent& MouseEvent);
+
+/**
+ * UI to display the provided UTexture and respond to mouse clicks. 
+ */
+class SSimulcamViewport : public SCompoundWidget
+{
+public:
+	SLATE_BEGIN_ARGS(SSimulcamViewport) { }
+	SLATE_ATTRIBUTE(FVector2D, BrushImageSize);
+	SLATE_EVENT(FSimulcamViewportClickedEventHandler, OnSimulcamViewportClicked)
+	SLATE_END_ARGS()
+
+public:
+
+	/** Default constructor. */
+	SSimulcamViewport();
+
+public:
+
+	/**
+	 * Construct this widget
+	 *
+	 * @param InArgs The declaration data for this widget.
+	 * @param InTexture The source texture to render if any
+	 */
+	void Construct(const FArguments& InArgs, UTexture* InTexture);
+
+	/**
+	 * Tick this widget
+	 * 
+	 * @param InAllottedGeometry Geometry of the widget.
+	 * @param InCurrentTime CurrentTime of the engine.
+	 * @param InDeltaTime Deltatime since last frame.
+	 */
+	virtual void Tick(const FGeometry& InAllottedGeometry, const double InCurrentTime, const float InDeltaTime) override;
+
+private:
+
+	class FInternalReferenceCollector : public FGCObject
+	{
+	public:
+		FInternalReferenceCollector(SSimulcamViewport* InObject)
+			: Object(InObject)
+		{
+		}
+
+		//~ FGCObject interface
+		virtual void AddReferencedObjects(FReferenceCollector& InCollector) override
+		{
+			InCollector.AddReferencedObject(Object->Material);
+			InCollector.AddReferencedObject(Object->TextureSampler);
+		}
+
+	private:
+		SSimulcamViewport* Object;
+	};
+
+	friend FInternalReferenceCollector;
+
+	/** Collector to keep the UObject alive. */
+	FInternalReferenceCollector Collector;
+
+	/** The material that wraps the video texture for display in an SImage. */
+	UMaterial* Material;
+
+	/** The Slate brush that renders the material. */
+	TSharedPtr<FSlateBrush> MaterialBrush;
+
+	/** The video texture sampler in the wrapper material. */
+	UMaterialExpressionTextureSample* TextureSampler;
+
+	/** Brush image size attribute. */
+	TAttribute<FVector2D> BrushImageSize;
+
+private:
+
+	/** Delegate to be executed when the viewport area is clicked */
+	FSimulcamViewportClickedEventHandler OnSimulcamViewportClicked;
+};
+
diff --git a/Engine/Plugins/VirtualProduction/CameraCalibration/Source/CameraCalibrationEditor/Private/Calibrators/CameraNodalOffsetAlgoPoints.cpp b/Engine/Plugins/VirtualProduction/CameraCalibration/Source/CameraCalibrationEditor/Private/Calibrators/CameraNodalOffsetAlgoPoints.cpp
new file mode 100644
--- /dev/null
+++ b/Engine/Plugins/VirtualProduction/CameraCalibration/Source/CameraCalibrationEditor/Private/Calibrators/CameraNodalOffsetAlgoPoints.cpp
@@ -0,0 +1,952 @@
+// Copyright Epic Games, Inc. All Rights Reserved.
+
+#pragma once
+
+#include "CameraNodalOffsetAlgoPoints.h"
+#include "AssetEditor/NodalOffsetTool.h"
+#include "CalibrationPointComponent.h"
+#include "Camera/CameraActor.h"
+#include "Camera/CameraComponent.h"
+#include "CameraCalibrationSubsystem.h"
+#include "CameraCalibrationUtils.h"
+#include "Editor.h"
+#include "GameFramework/Actor.h"
+#include "Input/Events.h"
+#include "Internationalization/Text.h"
+#include "Layout/Geometry.h"
+#include "LensFile.h"
+#include "Math/Vector.h"
+#include "MediaTexture.h"
+#include "Misc/MessageDialog.h"
+#include "PropertyCustomizationHelpers.h"
+#include "SphericalLensDistortionModelHandler.h"
+#include "UI/CameraCalibrationWidgetHelpers.h"
+#include "Widgets/Input/SButton.h"
+#include "Widgets/Layout/SSpacer.h"
+#include "Widgets/Views/SListView.h"
+#include "Widgets/SWidget.h"
+
+#include <vector>
+
+#if WITH_OPENCV
+#include "OpenCVHelper.h"
+OPENCV_INCLUDES_START
+// The check macro causes problems with opencv headers. 
+// e.g. OpenCV\include\opencv2/core/utility.hpp(53): fatal error C1021: invalid preprocessor command 'warning'
+#undef check 
+#include <opencv2/opencv.hpp>
+#include "opencv2/calib3d.hpp"
+OPENCV_INCLUDES_END
+#endif
+
+#define LOCTEXT_NAMESPACE "CameraNodalOffsetAlgoPoints"
+
+
+void UCameraNodalOffsetAlgoPoints::Initialize(TWeakPtr<FNodalOffsetTool> InNodalOffsetTool)
+{
+	NodalOffsetTool = InNodalOffsetTool;
+
+	// Guess which calibrator to use but searching for actors with CalibrationPointComponents.
+	SetCalibrator(FindFirstCalibrator());
+}
+
+void UCameraNodalOffsetAlgoPoints::Shutdown()
+{
+	NodalOffsetTool.Reset();
+}
+
+void UCameraNodalOffsetAlgoPoints::Tick(float DeltaTime)
+{
+	if (!NodalOffsetTool.IsValid())
+	{
+		return;
+	}
+
+	// If not paused, cache calibrator 3d point position
+	if (!NodalOffsetTool.Pin()->IsPaused())
+	{
+		// Get calibration point data
+		do 
+		{	
+			LastCalibratorPoint.bIsValid = false;
+
+			if (!CalibratorPointsComboBox.IsValid())
+			{
+				break;
+			}
+
+			TSharedPtr<FCalibratorPointData> CalibratorPoint = CalibratorPointsComboBox->GetSelectedItem();
+
+			if (!CalibratorPoint.IsValid())
+			{
+				break;
+			}
+
+			LastCalibratorPoint.Name = CalibratorPoint->Name;
+
+			if (!GetCurrentCalibratorPointLocation(LastCalibratorPoint.Location))
+			{
+				break;
+			}
+
+			LastCalibratorPoint.bIsValid = true;
+
+		} while (0);
+
+		// Get camera data
+		do
+		{
+			LastCameraData.bIsValid = false;
+
+			const FLensFileEvalData* LensFileEvalData = NodalOffsetTool.Pin()->GetLensFileEvalData();
+
+			// We require lens evaluation data, and that distortion was evaluated so that 2d correlations are valid
+			// Note: The comp enforces distortion application.
+			if (!LensFileEvalData || !LensFileEvalData->Distortion.bWasEvaluated)
+			{
+				break;
+			}
+
+			const ACameraActor* Camera = NodalOffsetTool.Pin()->GetCamera();
+
+			if (!Camera)
+			{
+				break;
+			}
+
+			const UCameraComponent* CameraComponent = Camera->GetCameraComponent();
+
+			if (!CameraComponent)
+			{
+				break;
+			}
+
+			LastCameraData.Pose = CameraComponent->GetComponentToWorld();
+			LastCameraData.UniqueId = Camera->GetUniqueID();
+			LastCameraData.LensFileEvalData = *LensFileEvalData;
+			LastCameraData.bIsValid = true;
+
+		} while (0);
+	}
+}
+
+bool UCameraNodalOffsetAlgoPoints::OnViewportClicked(const FGeometry& MyGeometry, const FPointerEvent& MouseEvent)
+{
+	// We only respond to left clicks
+	if (MouseEvent.GetEffectingButton() != EKeys::LeftMouseButton)
+	{
+		return false;
+	}
+
+	if (!NodalOffsetTool.IsValid())
+	{
+		return true;
+	}
+
+	// Check that we have a valid calibrator 3dpoint or camera data
+	if (!LastCalibratorPoint.bIsValid || !LastCameraData.bIsValid)
+	{
+		return true;
+	}
+
+	// Create the row that we're going to add
+	TSharedPtr<FCalibrationRowData> Row = MakeShared<FCalibrationRowData>();
+
+	Row->CameraData = LastCameraData;
+	Row->CalibratorPointData = LastCalibratorPoint;
+
+	// Get the mouse click 2d position
+	if (!CalculateNormalizedMouseClickPosition(MyGeometry, MouseEvent, Row->Point2D))
+	{
+		return true;
+	}
+
+	// Validate the new row, show a message if validation fails.
+	{
+		FText ErrorMessage;
+
+		if (!ValidateNewRow(Row, ErrorMessage))
+		{
+			FMessageDialog::Open(EAppMsgType::Ok, ErrorMessage);
+			return true;
+		}
+	}
+
+	// Add this data point
+	CalibrationRows.Add(Row);
+
+	// Notify the ListView of the new data
+	if (CalibrationListView.IsValid())
+	{
+		CalibrationListView->RequestListRefresh();
+	}
+
+	// Auto-advance to the next calibration point (if it exists)
+	AdvanceCalibratorPoint();
+
+	// Unpause, since we already picked a point.
+	if (NodalOffsetTool.Pin()->IsPaused())
+	{
+		NodalOffsetTool.Pin()->TogglePlay();
+	}
+
+	return true;
+}
+
+bool UCameraNodalOffsetAlgoPoints::CalculateNormalizedMouseClickPosition(const FGeometry& MyGeometry, const FPointerEvent& MouseEvent, FVector2D& OutPosition) const
+{
+	// Reject viewports with no area
+	if (FMath::IsNearlyZero(MyGeometry.Size.X) || FMath::IsNearlyZero(MyGeometry.Size.Y))
+	{
+		return false;
+	}
+
+	// About the Mouse Event data:
+	// 
+	// * MouseEvent.GetScreenSpacePosition(): Position in pixels on the screen (independent of window size of position)
+	// * MyGeometry.Size                    : Size of viewport (the one with the media, not the whole window)
+	// * MyGeometry.AbsolutePosition        : Position of the top-left corner of viewport within screen
+	// * MyGeometry.AbsoluteToLocal(MouseEvent.GetScreenSpacePosition()) gives you the pixel coordinates local to the viewport.
+
+	const FVector2D LocalInPixels = MyGeometry.AbsoluteToLocal(MouseEvent.GetScreenSpacePosition());
+
+	const float XNormalized = LocalInPixels.X / MyGeometry.Size.X;
+	const float YNormalized = LocalInPixels.Y / MyGeometry.Size.Y;
+
+	// Position 0~1. Origin at top-left corner of the viewport.
+	OutPosition = FVector2D(XNormalized, YNormalized);
+
+	return true;
+}
+
+void UCameraNodalOffsetAlgoPoints::AdvanceCalibratorPoint()
+{
+	for (int32 PointIdx = 0; PointIdx < CurrentCalibratorPoints.Num(); PointIdx++)
+	{
+		if (CurrentCalibratorPoints[PointIdx]->Name == LastCalibratorPoint.Name)
+		{
+			const int32 NextIdx = (PointIdx + 1) % CurrentCalibratorPoints.Num();
+			CalibratorPointsComboBox->SetSelectedItem(CurrentCalibratorPoints[NextIdx]);
+			break;
+		}
+	}
+}
+
+bool UCameraNodalOffsetAlgoPoints::GetCurrentCalibratorPointLocation(FVector& OutLocation)
+{
+	TSharedPtr<FCalibratorPointData> CalibratorPointData = CalibratorPointsComboBox->GetSelectedItem();
+
+	if (!CalibratorPointData.IsValid())
+	{
+		return false;
+	}
+	const UCalibrationPointComponent* CalibrationPointComponent = GetCalibrationPointComponentFromName(CalibratorPointData->Name);
+
+	if (!CalibrationPointComponent)
+	{
+		return false;
+	}
+
+	OutLocation = CalibrationPointComponent->GetComponentLocation();
+
+	return true;
+}
+
+TSharedRef<SWidget> UCameraNodalOffsetAlgoPoints::BuildUI()
+{
+	return SNew(SVerticalBox)
+
+		+ SVerticalBox::Slot() // Calibrator picker
+		.VAlign(EVerticalAlignment::VAlign_Top)
+		.AutoHeight()
+		.MaxHeight(FCameraCalibrationWidgetHelpers::DefaultRowHeight)
+		[ FCameraCalibrationWidgetHelpers::BuildLabelWidgetPair(LOCTEXT("Calibrator", "Calibrator"), BuildCalibrationDevicePickerWidget()) ]
+		
+		+SVerticalBox::Slot() // Calibrator point names
+		.AutoHeight()
+		.MaxHeight(FCameraCalibrationWidgetHelpers::DefaultRowHeight)
+		[ FCameraCalibrationWidgetHelpers::BuildLabelWidgetPair(LOCTEXT("CalibratorPoint", "Calibrator Point"), BuildCalibrationPointsComboBox()) ]
+		
+		+ SVerticalBox::Slot() // Calibration Rows
+		.AutoHeight()
+		.MaxHeight(12 * FCameraCalibrationWidgetHelpers::DefaultRowHeight)
+		[ BuildCalibrationPointsTable() ]
+		
+		+ SVerticalBox::Slot() // Action buttons (e.g. Remove, Clear)
+		.HAlign(EHorizontalAlignment::HAlign_Center)
+		.AutoHeight()
+		.Padding(0,20)
+		[ BuildCalibrationActionButtons() ]
+		;
+}
+
+bool UCameraNodalOffsetAlgoPoints::ValidateNewRow(TSharedPtr<FCalibrationRowData>& Row, FText& OutErrorMessage) const
+{
+	if (!CalibrationRows.Num())
+	{
+		return true;
+	}
+
+	if (!Row.IsValid())
+	{
+		OutErrorMessage = LOCTEXT("InvalidRowPointer", "Invalid row pointer");
+		return false;
+	}
+
+	if (!NodalOffsetTool.IsValid())
+	{
+		return false;
+	}
+
+	const ULensFile* LensFile = NodalOffsetTool.Pin()->GetLensFile();
+
+	if (!LensFile)
+	{
+		OutErrorMessage = LOCTEXT("LensFileNotFound", "LensFile not found");
+		return false;
+	}
+
+	// Same LensFile
+
+	const TSharedPtr<FCalibrationRowData>& FirstRow = CalibrationRows[0];
+
+	if (Row->CameraData.LensFileEvalData.LensFile != LensFile)
+	{
+		OutErrorMessage = LOCTEXT("LensFileWasNotTheSame", "Lens file was not the same");
+		return false;
+	}
+
+	// Same camera in view
+
+	const ACameraActor* Camera = NodalOffsetTool.Pin()->GetCamera();
+
+	if (!Camera || !Camera->GetCameraComponent())
+	{
+		OutErrorMessage = LOCTEXT("MissingCamera", "Missing camera");
+		return false;
+	}
+
+	if (Camera->GetUniqueID() != Row->CameraData.UniqueId)
+	{
+		OutErrorMessage = LOCTEXT("DifferentCameraAsSelected", "Different camera as selected");
+		return false;
+	}
+
+	// Same camera as before
+
+	if (FirstRow->CameraData.UniqueId != Row->CameraData.UniqueId)
+	{
+		OutErrorMessage = LOCTEXT("CameraChangedDuringTheTest", "Camera changed during the test");
+		return false;
+	}
+
+	// Camera did not move much
+	{
+		// Location check
+
+		const FVector LocationDelta = Row->CameraData.Pose.GetLocation() - FirstRow->CameraData.Pose.GetLocation();
+
+		const float MaxLocationDeltaInCm = 2;
+		const float LocationDeltaInCm = LocationDelta.Size();
+
+		if (LocationDeltaInCm > MaxLocationDeltaInCm)
+		{
+			FFormatOrderedArguments Arguments;
+			Arguments.Add(FText::FromString(FString::Printf(TEXT("%.1f"), MaxLocationDeltaInCm)));
+			Arguments.Add(FText::FromString(FString::Printf(TEXT("%.1f"), LocationDeltaInCm)));
+
+			OutErrorMessage = FText::Format(LOCTEXT("CameraMovedLocation", "Camera moved more than {0} cm during the calibration ({1} cm)"), Arguments);
+
+			return false;
+		}
+
+		// Rotation check
+
+		const float AngularDistanceRadians = FirstRow->CameraData.Pose.GetRotation().AngularDistance(Row->CameraData.Pose.GetRotation());
+		const float MaxAngularDistanceRadians = 2 * (180 / PI);
+
+		if (AngularDistanceRadians > MaxAngularDistanceRadians)
+		{
+			FFormatOrderedArguments Arguments;
+			Arguments.Add(FText::FromString(FString::Printf(TEXT("%.1f"), MaxAngularDistanceRadians)));
+			Arguments.Add(FText::FromString(FString::Printf(TEXT("%.1f"), AngularDistanceRadians)));
+
+			OutErrorMessage = FText::Format(LOCTEXT("CameraMovedRotation", "Camera moved more than {0} radians during the calibration ({1} radians)"), Arguments);
+
+			return false;
+		}
+	}
+
+	// FZ inputs are valid
+	if ((!Row->CameraData.LensFileEvalData.Input.Focus.IsSet()) || (!Row->CameraData.LensFileEvalData.Input.Zoom.IsSet()))
+	{
+		OutErrorMessage = LOCTEXT("LutInputsNotValid", "FZ Lut inputs are not valid. Make sure you are providing Focus and Zoom values via LiveLink");
+		return false;
+	}
+
+	// bApplyNodalOffset did not change.
+	//
+	// It can't change because we need to know if the camera pose is being affected or not by the current nodal offset evaluation.
+	// And we need to know that because the offset we calculate will need to either subtract or not the current evaluation when adding it to the LUT.
+	if (FirstRow->CameraData.LensFileEvalData.NodalOffset.bWasApplied != Row->CameraData.LensFileEvalData.NodalOffset.bWasApplied)
+	{
+		OutErrorMessage = LOCTEXT("ApplyNodalOffsetChanged", "Apply nodal offset changed");
+		return false;
+	}
+
+	//@todo Focus and zoom did not change much (i.e. inputs to distortion and nodal offset). 
+	//      Threshold for physical units should differ from normalized encoders.
+
+	return true;
+}
+
+bool UCameraNodalOffsetAlgoPoints::GetNodalOffset(FNodalPointOffset& OutNodalOffset, float& OutFocus, float& OutZoom, float& OutError, FText& OutErrorMessage)
+{
+	// Sanity checks
+	//
+
+	// Enough points
+	if (CalibrationRows.Num() < 4)
+	{
+		OutErrorMessage = LOCTEXT("NotEnoughSamples", "At least 4 correspondence points are required");
+		return false;
+	}
+
+	// All points are valid
+	for (const TSharedPtr<FCalibrationRowData>& Row : CalibrationRows)
+	{
+		if (!ensure(Row.IsValid()))
+		{
+			return false;
+		}
+	}
+
+	if (!NodalOffsetTool.IsValid())
+	{
+		return false;
+	}
+
+	const ULensFile* LensFile = NodalOffsetTool.Pin()->GetLensFile();
+
+	if (!LensFile)
+	{
+		OutErrorMessage = LOCTEXT("LensFileNotFound", "LensFile not found");
+		return false;
+	}
+
+	// Get camera.
+
+	const ACameraActor* Camera = NodalOffsetTool.Pin()->GetCamera();
+
+	if (!Camera || !Camera->GetCameraComponent())
+	{
+		OutErrorMessage = LOCTEXT("MissingCamera", "Missing camera");
+		return false;
+	}
+
+	UCameraComponent* CameraComponent = Camera->GetCameraComponent();
+	check(CameraComponent);
+
+	UCineCameraComponent* CineCameraComponent = Cast<UCineCameraComponent>(CameraComponent);
+
+	if (!CineCameraComponent)
+	{
+		OutErrorMessage = LOCTEXT("OnlyCineCamerasSupported", "Only cine cameras are supported");
+		return false;
+	}
+
+	TSharedPtr<FCalibrationRowData>& FirstRow = CalibrationRows[0];
+
+	// Still same camera (since we need it to get the distortion handler, which much be the same)
+
+	if (Camera->GetUniqueID() != FirstRow->CameraData.UniqueId)
+	{
+		OutErrorMessage = LOCTEXT("DifferentCameraAsSelected", "Different camera as selected");
+		return false;
+	}
+
+	//@todo Should cache the distortion instead of getting it live, since the user may have changed it.
+
+	// Only spherical lens distortion is currently supported at the moment.
+
+	const USphericalLensDistortionModelHandler* SphericalHandler = Cast<USphericalLensDistortionModelHandler>(NodalOffsetTool.Pin()->GetDistortionHandler());
+
+	if (!SphericalHandler)
+	{
+		OutErrorMessage = LOCTEXT("OnlySphericalDistortionSupported", "Only spherical distortion is currently supported.");
+		return false;
+	}
+
+	// Get parameters from the handler
+	FLensDistortionState DistortionState = SphericalHandler->GetCurrentDistortionState();
+
+	// Extract named distortion parameters
+	FSphericalDistortionParameters SphericalParameters;
+	USphericalLensModel::StaticClass()->GetDefaultObject<ULensModel>()->FromArray(DistortionState.DistortionInfo.Parameters, SphericalParameters);
+
+#if !WITH_OPENCV
+	{
+		OutErrorMessage = LOCTEXT("OpenCVRequired", "OpenCV is required");
+		return false;
+	}
+#endif
+
+#if WITH_OPENCV
+
+	// Find the pose that minimizes the reprojection error
+
+	// Populate the 3d/2d correlation points
+
+	std::vector<cv::Point3f> Points3d;
+	std::vector<cv::Point2f> Points2d;
+
+	for (const TSharedPtr<FCalibrationRowData>& Row : CalibrationRows)
+	{
+		// Convert from UE coordinates to OpenCV coordinates
+
+		FTransform Transform;
+		Transform.SetIdentity();
+		Transform.SetLocation(Row->CalibratorPointData.Location);
+
+		FCameraCalibrationUtils::ConvertUnrealToOpenCV(Transform);
+
+		// Calibrator 3d points
+		Points3d.push_back(cv::Point3f(
+			Transform.GetLocation().X, 
+			Transform.GetLocation().Y, 
+			Transform.GetLocation().Z));
+
+		// Image 2d points
+		Points2d.push_back(cv::Point2f(
+			Row->Point2D.X,
+			Row->Point2D.Y));
+	}
+
+	// Populate camera matrix
+
+	cv::Mat CameraMatrix(3, 3, cv::DataType<double>::type);
+	cv::setIdentity(CameraMatrix);
+
+	// Note: cv::Mat uses (row,col) indexing.
+	//
+	//  Fx  0  Cx
+	//  0  Fy  Cy
+	//  0   0   1
+
+	CameraMatrix.at<double>(0, 0) = DistortionState.DistortionInfo.FxFy.X;
+	CameraMatrix.at<double>(1, 1) = DistortionState.DistortionInfo.FxFy.Y;
+
+	CameraMatrix.at<double>(0, 2) = DistortionState.PrincipalPoint.X;
+	CameraMatrix.at<double>(1, 2) = DistortionState.PrincipalPoint.Y;
+
+	// Populate distortion coefficients
+	// Note: solvePnP expects k1, k2, p1, p2, k3
+
+	cv::Mat DistortionCoefficients(5, 1, cv::DataType<double>::type);
+
+	DistortionCoefficients.at<double>(0) = SphericalParameters.K1;
+	DistortionCoefficients.at<double>(1) = SphericalParameters.K2;
+	DistortionCoefficients.at<double>(2) = SphericalParameters.P1;
+	DistortionCoefficients.at<double>(3) = SphericalParameters.P2;
+	DistortionCoefficients.at<double>(4) = SphericalParameters.K3;
+
+	// Solve for camera position
+
+	cv::Mat Rrod = cv::Mat::zeros(3, 1, cv::DataType<double>::type); // Rotation vector in Rodrigues notation. 3x1.
+	cv::Mat Tobj = cv::Mat::zeros(3, 1, cv::DataType<double>::type); // Translation vector. 3x1.
+
+	if (!cv::solvePnP(Points3d, Points2d, CameraMatrix, DistortionCoefficients, Rrod, Tobj))
+	{
+		OutErrorMessage = LOCTEXT("SolvePnpFailed", "SolvePnP failed");
+		return false;
+	}
+
+	// Check for invalid data
+	{
+		const double Tx = Tobj.at<double>(0);
+		const double Ty = Tobj.at<double>(1);
+		const double Tz = Tobj.at<double>(2);
+
+		const double MaxValue = 1e16;
+
+		if (abs(Tx) > MaxValue || abs(Ty) > MaxValue || abs(Tz) > MaxValue)
+		{
+			OutErrorMessage = LOCTEXT("DataOutOfBounds", "Data Out Of Bounds");
+			return false;
+		}
+	}
+
+	// Convert to camera pose
+
+	// [R|t]' = [R'|-R'*t]
+
+	// Convert from Rodrigues to rotation matrix
+	cv::Mat Robj;
+	cv::Rodrigues(Rrod, Robj); // Robj is 3x3
+
+	// Calculate camera translation
+	cv::Mat Tcam = -Robj.t() * Tobj;
+
+	// Invert/transpose to get camera orientation
+	cv::Mat Rcam = Robj.t();
+
+	// Convert back to UE coordinates
+
+	FMatrix M = FMatrix::Identity;
+
+	// Fill rotation matrix
+	for (int32 Column = 0; Column < 3; ++Column)
+	{
+		M.SetColumn(Column, FVector(
+			Rcam.at<double>(Column, 0), 
+			Rcam.at<double>(Column, 1), 
+			Rcam.at<double>(Column, 2))
+		);
+	}
+
+	// Fill translation vector
+	M.M[3][0] = Tcam.at<double>(0);
+	M.M[3][1] = Tcam.at<double>(1);
+	M.M[3][2] = Tcam.at<double>(2);
+
+	FTransform DesiredCameraTransform;
+	DesiredCameraTransform.SetFromMatrix(M);
+
+	FCameraCalibrationUtils::ConvertOpenCVToUnreal(DesiredCameraTransform);
+
+	// This is how we update the offset even when the camera is evaluating the current
+	// nodal offset curve in the Lens File:
+	// 
+	// CameraTransform = ExistingOffset * CameraTransformWithoutOffset
+	// => CameraTransformWithoutOffset = ExistingOffset' * CameraTransform
+	//
+	// DesiredTransform = Offset * CameraTransformWithoutOffset
+	// => Offset = DesiredTransform * CameraTransformWithoutOffset'
+	// => Offset = DesiredTransform * (ExistingOffset' * CameraTransform)'
+	// => Offset = DesiredTransform * (CameraTransform' * ExistingOffset)
+
+	// Evaluate nodal offset
+
+	// Determine the input values to the LUT (focus and zoom)
+
+	checkSlow(FirstRow->CameraData.LensFileEvalData.Input.Focus.IsSet());
+	checkSlow(FirstRow->CameraData.LensFileEvalData.Input.Zoom.IsSet());
+
+	OutFocus = *FirstRow->CameraData.LensFileEvalData.Input.Focus;
+	OutZoom  = *FirstRow->CameraData.LensFileEvalData.Input.Zoom;
+
+	// See if the camera already had an offset applied, in which case we need to account for it.
+
+	FTransform ExistingOffset = FTransform::Identity;
+
+	if (FirstRow->CameraData.LensFileEvalData.NodalOffset.bWasApplied)
+	{
+		FNodalPointOffset NodalPointOffset;
+
+		if (LensFile->EvaluateNodalPointOffset(OutFocus, OutZoom, NodalPointOffset))
+		{
+			ExistingOffset.SetTranslation(NodalPointOffset.LocationOffset);
+			ExistingOffset.SetRotation(NodalPointOffset.RotationOffset);
+		}
+	}
+
+	FTransform DesiredOffset = DesiredCameraTransform * LastCameraData.Pose.Inverse() * ExistingOffset;
+
+	OutNodalOffset.LocationOffset = DesiredOffset.GetLocation();
+	OutNodalOffset.RotationOffset = DesiredOffset.GetRotation();
+
+	return true;
+
+#endif //WITH_OPENCV
+
+	OutErrorMessage = LOCTEXT("OpenCVNotAvailable", "OpenCV was not available");
+	return false;
+}
+
+TSharedRef<SWidget> UCameraNodalOffsetAlgoPoints::BuildCalibrationDevicePickerWidget()
+{
+	return SNew(SObjectPropertyEntryBox)
+		.AllowedClass(AActor::StaticClass()) //@todo Consider populating this by searching for actors with CalibrationPoint components.
+		.OnObjectChanged_Lambda([&](const FAssetData& AssetData) -> void
+		{
+			if (AssetData.IsValid())
+			{
+				SetCalibrator(Cast<AActor>(AssetData.GetAsset()));
+			}
+		})
+		.ObjectPath_Lambda([&]() -> FString
+		{
+			if (AActor* TheCalibrator = GetCalibrator())
+			{
+				FAssetData AssetData(TheCalibrator, true);
+				return AssetData.ObjectPath.ToString();
+			}
+
+			return TEXT("");
+		});
+}
+
+class SCalibrationRowGenerator 
+	: public SMultiColumnTableRow<TSharedPtr<UCameraNodalOffsetAlgoPoints::FCalibrationRowData>>
+{
+	using FCalibrationRowData = UCameraNodalOffsetAlgoPoints::FCalibrationRowData;
+
+public:
+	SLATE_BEGIN_ARGS(SCalibrationRowGenerator) {}
+
+	/** The list item for this row */
+	SLATE_ARGUMENT(TSharedPtr<FCalibrationRowData>, CalibrationRowData)
+
+	SLATE_END_ARGS()
+
+	void Construct(const FArguments& Args, const TSharedRef<STableViewBase>& OwnerTableView)
+	{
+		CalibrationRowData = Args._CalibrationRowData;
+
+		SMultiColumnTableRow<TSharedPtr<FCalibrationRowData>>::Construct(
+			FSuperRowType::FArguments()
+			.Padding(1.0f),
+			OwnerTableView
+		);
+	}
+
+	//~Begin SMultiColumnTableRow
+	virtual TSharedRef<SWidget> GenerateWidgetForColumn(const FName& ColumnName) override
+	{
+		if (ColumnName == TEXT("Name"))
+		{
+			const FString Text = CalibrationRowData->CalibratorPointData.Name;
+			return SNew(STextBlock).Text(FText::FromString(Text));
+		}
+
+		if (ColumnName == TEXT("Point2D"))
+		{
+			const FString Text = FString::Printf(TEXT("(%.2f, %.2f)"), 
+				CalibrationRowData->Point2D.X, 
+				CalibrationRowData->Point2D.Y);
+
+			return SNew(STextBlock).Text(FText::FromString(Text));
+		}
+
+		if (ColumnName == TEXT("Point3D"))
+		{
+			const FString Text = FString::Printf(TEXT("(%.0f, %.0f, %.0f)"),
+				CalibrationRowData->CalibratorPointData.Location.X,
+				CalibrationRowData->CalibratorPointData.Location.Y,
+				CalibrationRowData->CalibratorPointData.Location.Z);
+
+			return SNew(STextBlock).Text(FText::FromString(Text));
+		}
+
+		return SNullWidget::NullWidget;
+	}
+	//~End SMultiColumnTableRow
+
+
+private:
+	TSharedPtr<FCalibrationRowData> CalibrationRowData;
+};
+
+TSharedRef<SWidget> UCameraNodalOffsetAlgoPoints::BuildCalibrationActionButtons()
+{
+	return SNew(SHorizontalBox)
+
+		+ SHorizontalBox::Slot() // Button to remove last row
+		.AutoWidth()
+		[
+			SNew(SButton)
+			.Text(LOCTEXT("RemoveLast", "Remove Last"))
+			.HAlign(HAlign_Center)
+			.VAlign(VAlign_Center)
+			.OnClicked_Lambda([&]() -> FReply
+			{
+				if (CalibrationRows.Num())
+				{
+					CalibrationRows.RemoveAt(CalibrationRows.Num() - 1);
+					if (CalibrationListView.IsValid())
+					{
+						CalibrationListView->RequestListRefresh();
+					}
+				}
+
+				return FReply::Handled();
+			})
+		]
+
+		+ SHorizontalBox::Slot() // Button to clear all rows
+		.AutoWidth()
+		[ 
+			SNew(SButton)
+			.Text(LOCTEXT("ClearAll", "Clear All"))
+			.HAlign(HAlign_Center)
+			.VAlign(VAlign_Center)
+			.OnClicked_Lambda([&]() -> FReply
+			{
+				ClearCalibrationRows();
+				return FReply::Handled();
+			})
+		]
+		;
+}
+
+TSharedRef<SWidget> UCameraNodalOffsetAlgoPoints::BuildCalibrationPointsComboBox()
+{
+	CalibratorPointsComboBox = SNew(SComboBox<TSharedPtr<FCalibratorPointData>>)
+		.OptionsSource(&CurrentCalibratorPoints)
+		.OnGenerateWidget_Lambda([&](TSharedPtr<FCalibratorPointData> InOption) -> TSharedRef<SWidget>
+		{
+			return SNew(STextBlock).Text(FText::FromString(*InOption->Name));
+		})
+		.InitiallySelectedItem(nullptr)
+		[
+			SNew(STextBlock)
+			.Text_Lambda([&]() -> FText
+			{
+				if (CalibratorPointsComboBox.IsValid() && CalibratorPointsComboBox->GetSelectedItem().IsValid())
+				{
+					return FText::FromString(CalibratorPointsComboBox->GetSelectedItem()->Name);
+				}
+
+				return LOCTEXT("InvalidComboOption", "None");
+			})
+		];
+
+	// Update combobox from current calibrator
+	SetCalibrator(GetCalibrator());
+
+	return CalibratorPointsComboBox.ToSharedRef();
+}
+
+TSharedRef<SWidget> UCameraNodalOffsetAlgoPoints::BuildCalibrationPointsTable()
+{
+	CalibrationListView = SNew(SListView<TSharedPtr<FCalibrationRowData>>)
+		.ItemHeight(24)
+		.ListItemsSource(&CalibrationRows)
+		.OnGenerateRow_Lambda([&](TSharedPtr<FCalibrationRowData> InItem, const TSharedRef<STableViewBase>& OwnerTable) -> TSharedRef<ITableRow>
+		{
+			return SNew(SCalibrationRowGenerator, OwnerTable).CalibrationRowData(InItem);
+		})
+		.SelectionMode(ESelectionMode::SingleToggle)
+		.HeaderRow
+		(
+			SNew(SHeaderRow)
+
+			+ SHeaderRow::Column("Name")
+			.DefaultLabel(LOCTEXT("Name", "Name"))
+
+			+ SHeaderRow::Column("Point2D")
+			.DefaultLabel(LOCTEXT("Point2D", "Point2D"))
+
+			+ SHeaderRow::Column("Point3D")
+			.DefaultLabel(LOCTEXT("Point3D", "Point3D"))
+		);
+
+	return CalibrationListView.ToSharedRef();
+}
+
+UWorld* UCameraNodalOffsetAlgoPoints::GetWorld() const
+{
+	return GEditor ? GEditor->GetEditorWorldContext().World() : nullptr;
+}
+
+AActor* UCameraNodalOffsetAlgoPoints::FindFirstCalibrator() const
+{
+	// We find the first UCalibrationPointComponent object and return its actor owner.
+
+	const EObjectFlags ExcludeFlags = RF_ClassDefaultObject; // We don't want the calibrator CDOs.
+	const UWorld* World = GetWorld();
+
+	for (TObjectIterator<UCalibrationPointComponent> It(ExcludeFlags, true, EInternalObjectFlags::PendingKill); It; ++It)
+	{
+		AActor* Actor = It->GetOwner();
+
+		if (Actor && (Actor->GetWorld() == World))
+		{
+			return Actor;
+		}
+	}
+
+	return nullptr;
+}
+
+const UCalibrationPointComponent* UCameraNodalOffsetAlgoPoints::GetCalibrationPointComponentFromName(FString& Name) const
+{
+	if (!Calibrator.IsValid())
+	{
+		return nullptr;
+	}
+
+	TArray<UCalibrationPointComponent*, TInlineAllocator<4>> CalibrationPoints;
+	Calibrator->GetComponents<UCalibrationPointComponent, TInlineAllocator<4>>(CalibrationPoints);
+
+	for (const UCalibrationPointComponent* CalibrationPoint : CalibrationPoints)
+	{
+		if (CalibrationPoint->GetName() == Name)
+		{
+			return CalibrationPoint;
+		}
+	}
+
+	return nullptr;
+}
+
+void UCameraNodalOffsetAlgoPoints::SetCalibrator(AActor* InCalibrator)
+{
+	Calibrator = InCalibrator;
+
+	// Update the list of points
+
+	CurrentCalibratorPoints.Empty();
+
+	if (!Calibrator.IsValid())
+	{
+		return;
+	}
+
+	TArray<UCalibrationPointComponent*, TInlineAllocator<4>> CalibrationPoints;
+	Calibrator->GetComponents<UCalibrationPointComponent, TInlineAllocator<4>>(CalibrationPoints);
+
+	for (const UCalibrationPointComponent* CalibrationPoint : CalibrationPoints)
+	{
+		CurrentCalibratorPoints.Add(MakeShared<FCalibratorPointData>(CalibrationPoint->GetName()));
+	}
+
+	// Notify combobox
+
+	if (!CalibratorPointsComboBox)
+	{
+		return;
+	}
+
+	CalibratorPointsComboBox->RefreshOptions();
+
+	if (CurrentCalibratorPoints.Num())
+	{
+		CalibratorPointsComboBox->SetSelectedItem(CurrentCalibratorPoints[0]);
+	}
+	else
+	{
+		CalibratorPointsComboBox->SetSelectedItem(nullptr);
+	}
+}
+
+AActor* UCameraNodalOffsetAlgoPoints::GetCalibrator() const
+{
+	return Calibrator.Get();
+}
+
+void UCameraNodalOffsetAlgoPoints::OnSavedNodalOffset()
+{
+	// Since the nodal point was saved, there is no further use for the current samples.
+	ClearCalibrationRows();
+}
+
+void UCameraNodalOffsetAlgoPoints::ClearCalibrationRows()
+{
+	CalibrationRows.Empty();
+
+	if (CalibrationListView.IsValid())
+	{
+		CalibrationListView->RequestListRefresh();
+	}
+}
\ No newline at end of file
diff --git a/Engine/Plugins/VirtualProduction/CameraCalibration/Source/CameraCalibrationEditor/Private/Calibrators/CameraNodalOffsetAlgoPoints.h b/Engine/Plugins/VirtualProduction/CameraCalibration/Source/CameraCalibrationEditor/Private/Calibrators/CameraNodalOffsetAlgoPoints.h
new file mode 100644
--- /dev/null
+++ b/Engine/Plugins/VirtualProduction/CameraCalibration/Source/CameraCalibrationEditor/Private/Calibrators/CameraNodalOffsetAlgoPoints.h
@@ -0,0 +1,187 @@
+// Copyright Epic Games, Inc. All Rights Reserved.
+
+#pragma once
+
+#include "CameraNodalOffsetAlgo.h"
+
+#include "LensFile.h"
+
+#include "CameraNodalOffsetAlgoPoints.generated.h"
+
+struct FGeometry;
+struct FPointerEvent;
+struct FTransform;
+struct FVector;
+
+class FNodalOffsetTool;
+
+template <typename ItemType>
+class SListView;
+
+class SWidget;
+class UCalibrationPointComponent;
+class UMediaTexture;
+class UWorld;
+
+template<typename OptionType>
+class SComboBox;
+
+
+/** 
+ * Implements a nodal offset calibration algorithm. It uses 3d points (UCalibrationPointComponent) 
+ * specified in a selectable calibrator with features that the user can identify by clicking on the 
+ * simulcam viewport, and after enough points have been identified, it can calculate the nodal offset
+ * that needs to be applied to the associated camera in order to align its CG with the live action media plate.
+ * 
+ * Limitations:
+ * - Only supports Brown–Conrady model lens model (FSphericalDistortionParameters)
+ * - The camera or the lens should not be moved during the calibration of each nodal offset sample.
+ */
+UCLASS()
+class UCameraNodalOffsetAlgoPoints : public UCameraNodalOffsetAlgo
+{
+	GENERATED_BODY()
+	
+public:
+
+	//~ Begin CalibPointsNodalOffsetAlgo
+	virtual void Initialize(TWeakPtr<FNodalOffsetTool> InNodalOffsetTool) override;
+	virtual void Shutdown() override;
+	virtual void Tick(float DeltaTime) override;
+	virtual bool OnViewportClicked(const FGeometry& MyGeometry, const FPointerEvent& MouseEvent) override;
+	virtual TSharedRef<SWidget> BuildUI() override;
+	virtual bool GetNodalOffset(FNodalPointOffset& OutNodalOffset, float& OutFocus, float& OutZoom, float& OutError, FText& OutErrorMessage) override;
+	virtual FName FriendlyName() const override { return TEXT("Nodal Offset Points Method"); };
+	virtual void OnSavedNodalOffset() override;
+	//~ End CalibPointsNodalOffsetAlgo
+
+private:
+
+	// SCalibrationRowGenerator will need access to the row structures below.
+	friend class SCalibrationRowGenerator;
+
+	/** Holds item information a given calibrator point in the calibrator model */
+	struct FCalibratorPointData
+	{
+		FCalibratorPointData(FString&& InName)
+			: Name(InName)
+		{}
+
+		/** Name of the calibrator 3d point, as defined in the mesh */
+		FString Name;
+	};
+
+	/** Holds information of the identified calibrator 2d point for a given sample of a 2d-3d correlation */
+	struct FCalibratorPointCache
+	{
+		// True if the rest of the contents are valid.
+		bool bIsValid;
+
+		// The name of the 3d calibrator point, as defined in the mesh
+		FString Name;
+
+		// The world space 3d location of the point
+		FVector Location;
+	};
+
+	/** Holds information of the camera pose for a given sample of a 2d-3d correlation */
+	struct FCameraDataCache
+	{
+		// True if the rest of the contents are valid.
+		bool bIsValid;
+
+		// The unique id of the camera object. Used to detect camera selection changes during a calibration session.
+		uint32 UniqueId;
+
+		// The camera pose
+		FTransform Pose;
+
+		// The data used to evaluate the lens data in the camera for this sample
+		FLensFileEvalData LensFileEvalData;
+	};
+
+	/** Holds information of the calibrator 3d point for a given sample of a 2d-3d correlation */
+	struct FCalibrationRowData
+	{
+		// Normalized 0~1 2d location of the identified calibrator point in the media plate.
+		FVector2D Point2D;
+
+		// Holds information of the calibrator point data for this sample.
+		FCalibratorPointCache CalibratorPointData;
+
+		// Holds information of the camera data for this sample
+		FCameraDataCache CameraData;
+	};
+
+private:
+
+	/** The nodal offset tool controller */
+	TWeakPtr<FNodalOffsetTool> NodalOffsetTool;
+
+	/** The currently selected calibrator object. It is expected to contain one or more UCalibrationPointComponent in it */
+	TWeakObjectPtr<AActor> Calibrator;
+
+	/** Options source for the CalibratorPointsComboBox. Lists the calibrator points found in the selected calibrator object */
+	TArray<TSharedPtr<FCalibratorPointData>> CurrentCalibratorPoints;
+
+	/** Allows the selection of calibrator point that will be visually identified in the simulcam viewport */
+	TSharedPtr<SComboBox<TSharedPtr<FCalibratorPointData>>> CalibratorPointsComboBox;
+
+	/** Rows source for the CalibrationListView */
+	TArray<TSharedPtr<FCalibrationRowData>> CalibrationRows;
+
+	/** Displays the list of calibration points that will be used to calculate the nodal offset */
+	TSharedPtr<SListView<TSharedPtr<FCalibrationRowData>>> CalibrationListView;
+
+	/** Caches the last calibrator point 3d location.  Will hold last value before the nodal offset tool is paused */
+	FCalibratorPointCache LastCalibratorPoint;
+
+	/** Caches the last camera data.  Will hold last value before the nodal offset tool is paused */
+	FCameraDataCache LastCameraData;
+
+private:
+
+	/** Builds the UI of the calibration device picker */
+	TSharedRef<SWidget> BuildCalibrationDevicePickerWidget();
+
+	/** Builds the UI of the calibration points table */
+	TSharedRef<SWidget> BuildCalibrationPointsTable();
+
+	/** Buils the UI for selecting the current calibrator point */
+	TSharedRef<SWidget> BuildCalibrationPointsComboBox();
+
+	/** Builds the UI for the action buttons (RemoveLast, ClearAll) */
+	TSharedRef<SWidget> BuildCalibrationActionButtons();
+
+private:
+
+	/** Retrieves the world that should be used for the referenced and spawned object */
+	UWorld* GetWorld() const;
+
+	/** Returns the first calibrator object in the scene that it can find */
+	AActor* FindFirstCalibrator() const;
+
+	/** Sets the calibrator object to be used. Updates the selection picker. */
+	void SetCalibrator(AActor* InCalibrator);
+
+	/** Returns the currently selected calibrator object. */
+	AActor* GetCalibrator() const;
+
+	/** Clears the list of calibration sample rows */
+	void ClearCalibrationRows();
+
+	/** Retrieves by name the UCalibrationPointComponent of the currently selected calibrator */
+	const UCalibrationPointComponent* GetCalibrationPointComponentFromName(FString& Name) const;
+
+	/** Returns the world 3d location of the currently selected */
+	bool GetCurrentCalibratorPointLocation(FVector& OutLocation);
+
+	/** Selects the next available UCalibrationPointComponent of the currently selected calibrator object */
+	void AdvanceCalibratorPoint();
+
+	/** Calculates the normalized (0~1) coordinates in the simulcam viewport of the given mouse click */
+	bool CalculateNormalizedMouseClickPosition(const FGeometry& MyGeometry, const FPointerEvent& MouseEvent, FVector2D& OutPosition) const;
+
+	/** Validates a new calibration point to determine if it should be added as a new sample row */
+	bool ValidateNewRow(TSharedPtr<FCalibrationRowData>& Row, FText& OutErrorMessage) const;
+};
diff --git a/Engine/Plugins/VirtualProduction/CameraCalibration/Source/CameraCalibrationEditor/Private/CameraCalibrationEditorLog.h b/Engine/Plugins/VirtualProduction/CameraCalibration/Source/CameraCalibrationEditor/Private/CameraCalibrationEditorLog.h
new file mode 100644
--- /dev/null
+++ b/Engine/Plugins/VirtualProduction/CameraCalibration/Source/CameraCalibrationEditor/Private/CameraCalibrationEditorLog.h
@@ -0,0 +1,8 @@
+// Copyright Epic Games, Inc. All Rights Reserved.
+
+#pragma once
+
+#include "CoreMinimal.h"
+#include "Logging/LogMacros.h"
+
+DECLARE_LOG_CATEGORY_EXTERN(LogCameraCalibrationEditor, Log, All);
diff --git a/Engine/Plugins/VirtualProduction/CameraCalibration/Source/CameraCalibrationEditor/Private/UI/CameraCalibrationWidgetHelpers.cpp b/Engine/Plugins/VirtualProduction/CameraCalibration/Source/CameraCalibrationEditor/Private/UI/CameraCalibrationWidgetHelpers.cpp
new file mode 100644
--- /dev/null
+++ b/Engine/Plugins/VirtualProduction/CameraCalibration/Source/CameraCalibrationEditor/Private/UI/CameraCalibrationWidgetHelpers.cpp
@@ -0,0 +1,26 @@
+// Copyright Epic Games, Inc. All Rights Reserved.
+#pragma once
+
+#include "CameraCalibrationWidgetHelpers.h"
+
+#include "Internationalization/Text.h"
+#include "Widgets/SBoxPanel.h"
+#include "Widgets/SWidget.h"
+#include "Widgets/Text/STextBlock.h"
+
+const int32 FCameraCalibrationWidgetHelpers::DefaultRowHeight = 35;
+
+TSharedRef<SWidget> FCameraCalibrationWidgetHelpers::BuildLabelWidgetPair(FText&& Text, TSharedRef<SWidget> Widget)
+{
+	return SNew(SHorizontalBox)
+
+		+ SHorizontalBox::Slot()
+		.VAlign(EVerticalAlignment::VAlign_Center)
+		.Padding(5,5)
+		[SNew(STextBlock).Text(Text)]
+
+		+ SHorizontalBox::Slot()
+		.VAlign(EVerticalAlignment::VAlign_Center)
+		.Padding(5, 5)
+		[Widget];
+}
diff --git a/Engine/Plugins/VirtualProduction/CameraCalibration/Source/CameraCalibrationEditor/Private/UI/CameraCalibrationWidgetHelpers.h b/Engine/Plugins/VirtualProduction/CameraCalibration/Source/CameraCalibrationEditor/Private/UI/CameraCalibrationWidgetHelpers.h
new file mode 100644
--- /dev/null
+++ b/Engine/Plugins/VirtualProduction/CameraCalibration/Source/CameraCalibrationEditor/Private/UI/CameraCalibrationWidgetHelpers.h
@@ -0,0 +1,21 @@
+// Copyright Epic Games, Inc. All Rights Reserved.
+
+#pragma once
+
+#include "Templates/SharedPointer.h"
+
+class SWidget;
+class FText;
+
+/**
+ * Helpers to build UI used by the Camera Calibration modules.
+ */
+class FCameraCalibrationWidgetHelpers
+{
+public:
+	/** Builds a UI with a horizontal box with a lable on the left and the provided widget on the right */
+	static TSharedRef<SWidget> BuildLabelWidgetPair(FText&& Text, TSharedRef<SWidget> Widget);
+
+	/** Stores the default row height used throughout the Camera Calibration UI */
+	static const int32 DefaultRowHeight;
+};
\ No newline at end of file
diff --git a/Engine/Plugins/VirtualProduction/LiveLinkCamera/Source/LiveLinkCamera/Private/LiveLinkCameraController.cpp b/Engine/Plugins/VirtualProduction/LiveLinkCamera/Source/LiveLinkCamera/Private/LiveLinkCameraController.cpp
--- a/Engine/Plugins/VirtualProduction/LiveLinkCamera/Source/LiveLinkCamera/Private/LiveLinkCameraController.cpp
+++ b/Engine/Plugins/VirtualProduction/LiveLinkCamera/Source/LiveLinkCamera/Private/LiveLinkCameraController.cpp
@@ -13,8 +13,8 @@
 #include "ILiveLinkClient.h"
 #include "LensFile.h"
 #include "LiveLinkComponentController.h"
-#include "Logging/LogMacros.h"
 #include "LiveLinkLog.h"
+#include "Logging/LogMacros.h"
 #include "Roles/LiveLinkCameraRole.h"
 #include "Roles/LiveLinkCameraTypes.h"
 #include "UObject/EnterpriseObjectVersion.h"
@@ -38,11 +38,37 @@ ULiveLinkCameraController::ULiveLinkCameraController()
 
 void ULiveLinkCameraController::Tick(float DeltaTime, const FLiveLinkSubjectFrameData& SubjectData)
 {
+	// Invalidate the lens file evaluation data
+	LensFileEvalData.Invalidate();
+
+	if (UCameraComponent* CameraComponent = Cast<UCameraComponent>(AttachedComponent))
+	{
+		if (AActor* Camera = Cast<AActor>(AttachedComponent->GetOwner()))
+		{
+			LensFileEvalData.Camera.UniqueId = Camera->GetUniqueID();
+		}
+	}
+
 	const FLiveLinkCameraStaticData* StaticData = SubjectData.StaticData.Cast<FLiveLinkCameraStaticData>();
 	const FLiveLinkCameraFrameData* FrameData = SubjectData.FrameData.Cast<FLiveLinkCameraFrameData>();
 
 	if (StaticData && FrameData)
 	{
+		if (StaticData->bIsFocusDistanceSupported)
+		{
+			LensFileEvalData.Input.Focus = FrameData->FocusDistance;
+		}
+
+		if (StaticData->bIsApertureSupported)
+		{
+			LensFileEvalData.Input.Iris = FrameData->Aperture;
+		}
+
+		if (StaticData->bIsFocalLengthSupported)
+		{
+			LensFileEvalData.Input.Zoom = FrameData->FocalLength;
+		}
+
 		if (UCameraComponent* CameraComponent = Cast<UCameraComponent>(AttachedComponent))
 		{
 			bIsEncoderMappingNeeded = (StaticData->FIZDataMode == ECameraFIZMode::EncoderData);
@@ -62,6 +88,7 @@ void ULiveLinkCameraController::Tick(float DeltaTime, const FLiveLinkSubjectFram
 				if (StaticData->FilmBackHeight > 0.0f && UpdateFlags.bApplyFilmBack) { CineCameraComponent->Filmback.SensorHeight = StaticData->FilmBackHeight; }
 				
 				ULensFile* SelectedLensFile = LensFilePicker.GetLensFile();
+				LensFileEvalData.LensFile = SelectedLensFile;
 
 				if (LastFilmback != CineCameraComponent->Filmback)
 				{
@@ -163,6 +190,30 @@ void ULiveLinkCameraController::PostEditImport()
 	}
 }
 
+void ULiveLinkCameraController::SetApplyNodalOffset(bool bInApplyNodalOffset)
+{
+	if (bApplyNodalOffset == bInApplyNodalOffset)
+	{
+		return;
+	}
+
+	bApplyNodalOffset = bInApplyNodalOffset;
+
+	if (UCineCameraComponent* const CineCameraComponent = Cast<UCineCameraComponent>(AttachedComponent))
+	{
+		if (bApplyNodalOffset)
+		{
+			OriginalCameraRotation = CineCameraComponent->GetRelativeRotation();
+			OriginalCameraLocation = CineCameraComponent->GetRelativeLocation();
+		}
+		else
+		{
+			CineCameraComponent->SetRelativeLocation(OriginalCameraLocation);
+			CineCameraComponent->SetRelativeRotation(OriginalCameraRotation.Quaternion());
+		}
+	}
+}
+
 #if WITH_EDITOR
 void ULiveLinkCameraController::PostEditChangeProperty(struct FPropertyChangedEvent& PropertyChangedEvent)
 {
@@ -213,11 +264,11 @@ void ULiveLinkCameraController::ApplyFIZ(ULensFile* LensFile, UCineCameraCompone
 		bool bUseRange = !bUseLensFileForEncoderMapping;
 		if (LensFile)
 		{
-			if (StaticData->bIsFocusDistanceSupported && UpdateFlags.bApplyFocusDistance)
+			if (LensFileEvalData.Input.Focus.IsSet() && UpdateFlags.bApplyFocusDistance)
 			{
 				if (LensFile->HasFocusEncoderMapping())
 				{
-					CineCameraComponent->FocusSettings.ManualFocusDistance = LensFile->EvaluateNormalizedFocus(FrameData->FocusDistance);
+					CineCameraComponent->FocusSettings.ManualFocusDistance = LensFile->EvaluateNormalizedFocus(*LensFileEvalData.Input.Focus);
 				}
 				else
 				{
@@ -225,23 +276,23 @@ void ULiveLinkCameraController::ApplyFIZ(ULensFile* LensFile, UCineCameraCompone
 				}
 			}
 
-			if (StaticData->bIsApertureSupported && UpdateFlags.bApplyAperture)
+			if (LensFileEvalData.Input.Iris.IsSet() && UpdateFlags.bApplyAperture)
 			{
 				if (LensFile->HasIrisEncoderMapping())
 				{
-					CineCameraComponent->CurrentAperture = LensFile->EvaluateNormalizedIris(FrameData->Aperture);
+					CineCameraComponent->CurrentAperture = LensFile->EvaluateNormalizedIris(*LensFileEvalData.Input.Iris);
 				}
 				else
 				{
-					UE_LOG(LogLiveLinkCameraController, Verbose, TEXT("'%s' could not evaluate raw iris value '%0.3f' using LensFile '%s'"), *GetName(), FrameData->FocusDistance, *LensFile->GetName())
+					UE_LOG(LogLiveLinkCameraController, Verbose, TEXT("'%s' could not evaluate raw iris value '%0.3f' using LensFile '%s'"), *GetName(), FrameData->Aperture, *LensFile->GetName())
 				}
 			}
 
-			if (StaticData->bIsFocalLengthSupported && UpdateFlags.bApplyFocalLength)
+			if (LensFileEvalData.Input.Zoom.IsSet() && UpdateFlags.bApplyFocalLength)
 			{
 				if (LensFile->HasZoomEncoderMapping())
 				{
-					CineCameraComponent->SetCurrentFocalLength(LensFile->EvaluateNormalizedZoom(FrameData->FocalLength));
+					CineCameraComponent->SetCurrentFocalLength(LensFile->EvaluateNormalizedZoom(*LensFileEvalData.Input.Zoom));
 				}
 				else
 				{
@@ -285,9 +336,20 @@ void ULiveLinkCameraController::ApplyFIZ(ULensFile* LensFile, UCineCameraCompone
 	}
 	else
 	{
-		if (StaticData->bIsFocusDistanceSupported && UpdateFlags.bApplyFocusDistance) { CineCameraComponent->FocusSettings.ManualFocusDistance = FrameData->FocusDistance; }
-		if (StaticData->bIsApertureSupported && UpdateFlags.bApplyAperture) { CineCameraComponent->CurrentAperture = FrameData->Aperture; }
-		if (StaticData->bIsFocalLengthSupported && UpdateFlags.bApplyFocalLength) { CineCameraComponent->SetCurrentFocalLength(FrameData->FocalLength); }
+		if (StaticData->bIsFocusDistanceSupported && UpdateFlags.bApplyFocusDistance) 
+		{ 
+			CineCameraComponent->FocusSettings.ManualFocusDistance = FrameData->FocusDistance; 
+		}
+
+		if (StaticData->bIsApertureSupported && UpdateFlags.bApplyAperture) 
+		{ 
+			CineCameraComponent->CurrentAperture = FrameData->Aperture; 
+		}
+
+		if (StaticData->bIsFocalLengthSupported && UpdateFlags.bApplyFocalLength)
+		{ 
+			CineCameraComponent->SetCurrentFocalLength(FrameData->FocalLength); 
+		}
 	}
 }
 
@@ -336,8 +398,14 @@ void ULiveLinkCameraController::ApplyNodalOffset(ULensFile* SelectedLensFile, UC
 		if (bApplyNodalOffset && SelectedLensFile)
 		{
 			FNodalPointOffset Offset;
-			if (SelectedLensFile->EvaluateNodalPointOffset(CineCameraComponent->CurrentFocusDistance, CineCameraComponent->CurrentFocalLength, Offset))
+
+			if (SelectedLensFile->EvaluateNodalPointOffset(
+				LensFileEvalData.Input.Focus.IsSet() ? *LensFileEvalData.Input.Focus : CineCameraComponent->CurrentFocusDistance, 
+				LensFileEvalData.Input.Zoom.IsSet() ? *LensFileEvalData.Input.Zoom : CineCameraComponent->CurrentFocalLength,
+				Offset))
 			{
+				LensFileEvalData.NodalOffset.bWasApplied = true;
+
 				CineCameraComponent->SetRelativeLocation(OriginalCameraLocation);
 				CineCameraComponent->SetRelativeRotation(OriginalCameraRotation.Quaternion());
 				CineCameraComponent->AddLocalOffset(Offset.LocationOffset);
@@ -345,7 +413,7 @@ void ULiveLinkCameraController::ApplyNodalOffset(ULensFile* SelectedLensFile, UC
 			}
 		}
 
-		LastLocation= CineCameraComponent->GetRelativeLocation();
+		LastLocation = CineCameraComponent->GetRelativeLocation();
 		LastRotation = CineCameraComponent->GetRelativeRotation();
 	}
 	else 
@@ -357,7 +425,7 @@ void ULiveLinkCameraController::ApplyNodalOffset(ULensFile* SelectedLensFile, UC
 
 void ULiveLinkCameraController::ApplyDistortion(ULensFile* LensFile, UCineCameraComponent* CineCameraComponent, const FLiveLinkCameraStaticData* StaticData, const FLiveLinkCameraFrameData* FrameData)
 {
-	bool bCanUpdateDistortion = (StaticData->bIsFocusDistanceSupported || StaticData->bIsFocalLengthSupported || StaticData->bIsFieldOfViewSupported);
+	const bool bCanUpdateDistortion = (StaticData->bIsFocusDistanceSupported || StaticData->bIsFocalLengthSupported || StaticData->bIsFieldOfViewSupported);
 
 	if (LensFile && bCanUpdateDistortion)
 	{
@@ -369,12 +437,51 @@ void ULiveLinkCameraController::ApplyDistortion(ULensFile* LensFile, UCineCamera
 
 		if (LensDistortionHandler)
 		{
+			// Adjust FocalLength and Filmback to match FxFy (which has already been divided by resolution in pixels)
+			{
+				FDistortionInfo DistortionInfo;
+				DistortionInfo.Parameters.SetNumZeroed(LensFile->LensInfo.LensModel.GetDefaultObject()->GetNumParameters());
+
+				const bool bGotFxFy = LensFile->EvaluateDistortionParameters(
+					LensFileEvalData.Input.Focus.IsSet() ? *LensFileEvalData.Input.Focus : CineCameraComponent->CurrentFocusDistance,
+					LensFileEvalData.Input.Zoom.IsSet() ? *LensFileEvalData.Input.Zoom : CineCameraComponent->CurrentFocalLength,
+					DistortionInfo
+				);
+
+				if (bGotFxFy && (DistortionInfo.FxFy[0] > KINDA_SMALL_NUMBER) && (DistortionInfo.FxFy[1] > KINDA_SMALL_NUMBER))
+				{
+					// This is how field of view, filmback, and focal length are related:
+					//
+					// FOVx = 2*atan(1/(2*Fx)) = 2*atan(FilmbackX / (2*FocalLength))
+					// => FocalLength = Fx*FilmbackX
+					// 
+					// FOVy = 2*atan(1/(2*Fy)) = 2*atan(FilmbackY / (2*FocalLength))
+					// => FilmbackY = FocalLength / Fy
+
+					CineCameraComponent->CurrentFocalLength = DistortionInfo.FxFy[0] * CineCameraComponent->Filmback.SensorWidth;
+					CineCameraComponent->Filmback.SensorHeight = CineCameraComponent->CurrentFocalLength / DistortionInfo.FxFy[1];
+				}
+			}
+
 			//Go through the lens file to get distortion data based on FIZ
 			//Our handler's displacement map will get updated
 			FDistortionData DistortionData;
-			const FVector2D CurrentSensorDimensions = FVector2D(CineCameraComponent->Filmback.SensorWidth, CineCameraComponent->Filmback.SensorHeight);
 
-			LensFile->EvaluateDistortionData(CineCameraComponent->CurrentFocusDistance, CineCameraComponent->CurrentFocalLength, CurrentSensorDimensions, LensDistortionHandler, DistortionData);
+			const FVector2D CurrentSensorDimensions = FVector2D(
+				CineCameraComponent->Filmback.SensorWidth, 
+				CineCameraComponent->Filmback.SensorHeight
+			);
+			
+			// Cache Lens evaluation data
+			LensFileEvalData.Distortion.bWasEvaluated = true;
+
+			LensFile->EvaluateDistortionData(
+				LensFileEvalData.Input.Focus.IsSet() ? *LensFileEvalData.Input.Focus : CineCameraComponent->CurrentFocusDistance,
+				LensFileEvalData.Input.Zoom.IsSet() ? *LensFileEvalData.Input.Zoom : CineCameraComponent->CurrentFocalLength,
+				CurrentSensorDimensions, 
+				LensDistortionHandler, 
+				DistortionData
+			);
 		}
 	}
 }
@@ -433,3 +540,8 @@ void ULiveLinkCameraController::PostLoad()
 	}
 #endif //WITH_EDITOR
 }
+
+const FLensFileEvalData& ULiveLinkCameraController::GetLensFileEvalDataRef() const
+{
+	return LensFileEvalData;
+}
diff --git a/Engine/Plugins/VirtualProduction/LiveLinkCamera/Source/LiveLinkCamera/Public/LiveLinkCameraController.h b/Engine/Plugins/VirtualProduction/LiveLinkCamera/Source/LiveLinkCamera/Public/LiveLinkCameraController.h
--- a/Engine/Plugins/VirtualProduction/LiveLinkCamera/Source/LiveLinkCamera/Public/LiveLinkCameraController.h
+++ b/Engine/Plugins/VirtualProduction/LiveLinkCamera/Source/LiveLinkCamera/Public/LiveLinkCameraController.h
@@ -82,6 +82,12 @@ public:
 	/** Returns true if encoder mapping is required for FIZ data */
 	bool IsEncoderMappingNeeded() const { return bIsEncoderMappingNeeded; }
 
+	/** Returns a const reference to input data used to evaluate the lens file */
+	const FLensFileEvalData& GetLensFileEvalDataRef() const;
+
+	/** Enables/disables the application of the nodal offset to the camera component */
+	void SetApplyNodalOffset(bool bInApplyNodalOffset);
+
 protected:
 	/** Applies FIZ data coming from LiveLink stream. Lens file is used if encoder mapping is required  */
 	void ApplyFIZ(ULensFile* LensFile, UCineCameraComponent* CineCameraComponent, const FLiveLinkCameraStaticData* StaticData, const FLiveLinkCameraFrameData* FrameData);
@@ -143,6 +149,11 @@ protected:
 	bool bShouldUpdateVisualComponentOnChange = true;
 #endif
 
+protected:
+
+	/** Caches the latest inputs to the LensFile evaluation */
+	FLensFileEvalData LensFileEvalData;
+
 private:
 	/** Whether incoming data requires encoder mapping */
 	bool bIsEncoderMappingNeeded = false;
