CameraCalibration: Added missing undef LOCTEXT_NAMESPACE

diff --git a/Engine/Plugins/VirtualProduction/CameraCalibration/Source/CameraCalibrationEditor/Private/AssetEditor/NodalOffsetTool.cpp b/Engine/Plugins/VirtualProduction/CameraCalibration/Source/CameraCalibrationEditor/Private/AssetEditor/NodalOffsetTool.cpp
--- a/Engine/Plugins/VirtualProduction/CameraCalibration/Source/CameraCalibrationEditor/Private/AssetEditor/NodalOffsetTool.cpp
+++ b/Engine/Plugins/VirtualProduction/CameraCalibration/Source/CameraCalibrationEditor/Private/AssetEditor/NodalOffsetTool.cpp
@@ -873,4 +873,6 @@ void FNodalOffsetTool::FindMediaSourceUrls(TArray<TSharedPtr<FString>>& OutMedia
 			OutMediaSourceUrls.Add(MakeShareable(new FString(TSMediaSource->GetUrl())));
 		}
 	}
-}
\ No newline at end of file
+}
+
+#undef LOCTEXT_NAMESPACE
diff --git a/Engine/Plugins/VirtualProduction/CameraCalibration/Source/CameraCalibrationEditor/Private/Calibrators/CameraNodalOffsetAlgoPoints.cpp b/Engine/Plugins/VirtualProduction/CameraCalibration/Source/CameraCalibrationEditor/Private/Calibrators/CameraNodalOffsetAlgoPoints.cpp
--- a/Engine/Plugins/VirtualProduction/CameraCalibration/Source/CameraCalibrationEditor/Private/Calibrators/CameraNodalOffsetAlgoPoints.cpp
+++ b/Engine/Plugins/VirtualProduction/CameraCalibration/Source/CameraCalibrationEditor/Private/Calibrators/CameraNodalOffsetAlgoPoints.cpp
@@ -949,4 +949,6 @@ void UCameraNodalOffsetAlgoPoints::ClearCalibrationRows()
 	{
 		CalibrationListView->RequestListRefresh();
 	}
-}
\ No newline at end of file
+}
+
+#undef LOCTEXT_NAMESPACE
