Fix duplicate log extern declaration
#rnx

diff --git a/Engine/Plugins/VirtualProduction/CameraCalibration/Source/CameraCalibrationEditor/Private/CameraCalibrationEditorModule.cpp b/Engine/Plugins/VirtualProduction/CameraCalibration/Source/CameraCalibrationEditor/Private/CameraCalibrationEditorModule.cpp
--- a/Engine/Plugins/VirtualProduction/CameraCalibration/Source/CameraCalibrationEditor/Private/CameraCalibrationEditorModule.cpp
+++ b/Engine/Plugins/VirtualProduction/CameraCalibration/Source/CameraCalibrationEditor/Private/CameraCalibrationEditorModule.cpp
@@ -4,6 +4,7 @@
 #include "CameraCalibrationEditorModule.h"
 
 #include "AssetEditor/CameraCalibrationCommands.h"
+#include "CameraCalibrationEditorLog.h"
 #include "AssetToolsModule.h"
 #include "AssetTypeActions/AssetTypeActions_LensFile.h"
 #include "CameraCalibrationTypes.h"
diff --git a/Engine/Plugins/VirtualProduction/CameraCalibration/Source/CameraCalibrationEditor/Private/CameraCalibrationEditorModule.h b/Engine/Plugins/VirtualProduction/CameraCalibration/Source/CameraCalibrationEditor/Private/CameraCalibrationEditorModule.h
--- a/Engine/Plugins/VirtualProduction/CameraCalibration/Source/CameraCalibrationEditor/Private/CameraCalibrationEditorModule.h
+++ b/Engine/Plugins/VirtualProduction/CameraCalibration/Source/CameraCalibrationEditor/Private/CameraCalibrationEditorModule.h
@@ -7,9 +7,6 @@
 
 class IAssetTypeActions;
 
-DECLARE_LOG_CATEGORY_EXTERN(LogCameraCalibrationEditor, Log, All);
-
-
 /**
  * Implements the CameraCalibrationEditor module.
  */
