Add an importer for ulens files.

This makes it possible to associate a LensFile asset with a file on disk, so it can be reimported as well

diff --git a/Engine/Plugins/VirtualProduction/CameraCalibration/Source/CameraCalibrationEditor/Private/AssetTypeActions/AssetTypeActions_LensFile.cpp b/Engine/Plugins/VirtualProduction/CameraCalibration/Source/CameraCalibrationEditor/Private/AssetTypeActions/AssetTypeActions_LensFile.cpp
--- a/Engine/Plugins/VirtualProduction/CameraCalibration/Source/CameraCalibrationEditor/Private/AssetTypeActions/AssetTypeActions_LensFile.cpp
+++ b/Engine/Plugins/VirtualProduction/CameraCalibration/Source/CameraCalibrationEditor/Private/AssetTypeActions/AssetTypeActions_LensFile.cpp
@@ -3,6 +3,7 @@
 #include "AssetTypeActions/AssetTypeActions_LensFile.h"
 
 #include "AssetEditor/CameraCalibrationToolkit.h"
+#include "EditorFramework/AssetImportData.h"
 #include "LensFile.h"
 
 #define LOCTEXT_NAMESPACE "LensFileTypeActions"
@@ -30,4 +31,18 @@ void FAssetTypeActions_LensFile::OpenAssetEditor(const TArray<UObject*>& InObjec
 	}
 }
 
+void FAssetTypeActions_LensFile::GetResolvedSourceFilePaths(const TArray<UObject*>& TypeAssets, TArray<FString>& OutSourceFilePaths) const
+{
+	for (UObject* Asset : TypeAssets)
+	{
+		if (const ULensFile* LensFile = CastChecked<ULensFile>(Asset))
+		{
+			if (LensFile->AssetImportData)
+			{
+				LensFile->AssetImportData->ExtractFilenames(OutSourceFilePaths);
+			}
+		}
+	}
+}
+
 #undef LOCTEXT_NAMESPACE
diff --git a/Engine/Plugins/VirtualProduction/CameraCalibration/Source/CameraCalibrationEditor/Private/AssetTypeActions/AssetTypeActions_LensFile.h b/Engine/Plugins/VirtualProduction/CameraCalibration/Source/CameraCalibrationEditor/Private/AssetTypeActions/AssetTypeActions_LensFile.h
--- a/Engine/Plugins/VirtualProduction/CameraCalibration/Source/CameraCalibrationEditor/Private/AssetTypeActions/AssetTypeActions_LensFile.h
+++ b/Engine/Plugins/VirtualProduction/CameraCalibration/Source/CameraCalibrationEditor/Private/AssetTypeActions/AssetTypeActions_LensFile.h
@@ -10,12 +10,29 @@ class FAssetTypeActions_LensFile : public FAssetTypeActions_Base
 {
 public:
 	//~ Begin IAssetTypeActions interface
-	virtual FColor GetTypeColor() const override { return FColor(100, 255, 100); }
-	virtual bool HasActions(const TArray<UObject*>& InObjects) const override { return false; }
-	virtual uint32 GetCategories() override { return EAssetTypeCategories::Misc; }
-	virtual bool IsImportedAsset() const override { return false; }
+	virtual FColor GetTypeColor() const override
+	{
+		return FColor(100, 255, 100);
+	}
+
+	virtual bool HasActions(const TArray<UObject*>& InObjects) const override
+	{
+		return false;
+	}
+
+	virtual uint32 GetCategories() override
+	{
+		return EAssetTypeCategories::Misc;
+	}
+
+	virtual bool IsImportedAsset() const override
+	{
+		return true;
+	}
+
 	virtual FText GetName() const override;
 	virtual UClass* GetSupportedClass() const override;
 	virtual void OpenAssetEditor(const TArray<UObject*>& InObjects, TSharedPtr<IToolkitHost> EditWithinLevelEditor = TSharedPtr<IToolkitHost>()) override;
+	virtual void GetResolvedSourceFilePaths(const TArray<UObject*>& TypeAssets, TArray<FString>& OutSourceFilePaths) const override;
 	//~ End IAssetTypeActions interface
 };
\ No newline at end of file
diff --git a/Engine/Plugins/VirtualProduction/CameraCalibration/Source/CameraCalibrationEditor/Private/Exporters/LensFileExporter.cpp b/Engine/Plugins/VirtualProduction/CameraCalibration/Source/CameraCalibrationEditor/Private/Exporters/LensFileExporter.cpp
--- a/Engine/Plugins/VirtualProduction/CameraCalibration/Source/CameraCalibrationEditor/Private/Exporters/LensFileExporter.cpp
+++ b/Engine/Plugins/VirtualProduction/CameraCalibration/Source/CameraCalibrationEditor/Private/Exporters/LensFileExporter.cpp
@@ -2,313 +2,13 @@
 
 #include "LensFileExporter.h"
 #include "JsonObjectConverter.h"
-#include "CameraCalibrationEditorLog.h"
-#include "CameraCalibrationUtils.h"
 #include "LensFile.h"
-#include "LensData.h"
+#include "LensFileExchangeFormat.h"
 #include "Models/LensModel.h"
 
 
-FLensInfoExchange::FLensInfoExchange(const ULensFile* LensFile)
-{
-	if (LensFile == nullptr)
-	{
-		return;
-	}
-
-	const FLensInfo& LensInfo = LensFile->LensInfo;
-
-	SerialNumber = *LensInfo.LensSerialNumber;
-	ModelName = *LensInfo.LensModelName;
-
-	if (const ULensModel* LensModelObject = LensInfo.LensModel->GetDefaultObject<ULensModel>())
-	{
-		DistortionModel = LensModelObject->GetShortModelName();
-	}
-}
-
-FLensFileMetadata::FLensFileMetadata(const ULensFile* LensFile)
-	: LensInfo{ LensFile }
-{
-	Type = TEXT("LensFile");
-	Version = FLensFileExchange::LensFileVersion;
-
-	if (LensFile == nullptr)
-	{
-		return;
-	}
-
-	// Set the default values for metadata when exporting from the ULensFile
-	Name = *LensFile->GetName();
-
-	// Add the UserMetadata
-	UserMetadata.Reserve(LensFile->UserMetadata.Num());
-	for (const TPair<FString, FString>& UserMetadataPair : LensFile->UserMetadata)
-	{
-		FLensFileUserMetadataEntry UserMetadataEntry;
-		UserMetadataEntry.Name = *UserMetadataPair.Key;
-		UserMetadataEntry.Value = *UserMetadataPair.Value;
-		UserMetadata.Add(UserMetadataEntry);
-	}
-}
-
-FLensFileExchange::FLensFileExchange(const ULensFile* LensFile)
-	: Metadata{ LensFile }
-{
-	if (LensFile == nullptr)
-	{
-		return;
-	}
-
-	const FLensInfo& LensInfo = LensFile->LensInfo;
-	const FVector2D& Sensor = LensInfo.SensorDimensions;
-
-	SensorDimensions.Width = Sensor.X;
-	SensorDimensions.Height = Sensor.Y;
-	SensorDimensions.Units = ELensFileUnit::Millimeters;
-
-	ImageDimensions.Width = LensInfo.ImageDimensions.X;
-	ImageDimensions.Height = LensInfo.ImageDimensions.Y;
-
-	ExtractFocalLengthTable(LensFile);
-	ExtractImageCenterTable(LensFile);
-	ExtractNodalOffsetTable(LensFile);
-	ExtractEncoderTables(LensFile);
-	ExtractDistortionParameters(LensFile);
-	ExtractSTMaps(LensFile);
-}
-
-void FLensFileExchange::ExtractFocalLengthTable(const ULensFile* LensFile)
-{
-	// Read the FocalLength table
-	FLensFileParameterTable FocalLengthParametersTable;
-	FocalLengthParametersTable.ParameterName = GET_MEMBER_NAME_CHECKED(ULensFile, FocalLengthTable);
-	FocalLengthParametersTable.Header = { FocusEncoderHeaderName, ZoomEncoderHeaderName, FocalLengthFxHeaderName, FocalLengthFyHeaderName };
-
-	const FFocalLengthTable& LensFileFocalLengthTable = LensFile->FocalLengthTable;
-
-	for (int32 FocusPointIndex = 0; FocusPointIndex < LensFileFocalLengthTable.GetFocusPointNum(); ++FocusPointIndex)
-	{
-		const FFocalLengthFocusPoint& FocalLengthFocusPoint = LensFileFocalLengthTable.FocusPoints[FocusPointIndex];
-		const float Focus = FocalLengthFocusPoint.GetFocus();
-
-		for (int32 FocalLengthInfoIndex = 0; FocalLengthInfoIndex < FocalLengthFocusPoint.GetNumPoints(); ++FocalLengthInfoIndex)
-		{
-			FFocalLengthInfo FocalLengthInfo;
-			if (FocalLengthFocusPoint.GetValue(FocalLengthInfoIndex, FocalLengthInfo))
-			{
-				const float Zoom = FocalLengthFocusPoint.GetZoom(FocalLengthInfoIndex);
-				const float Fx = FocalLengthInfo.FxFy.X;
-				const float Fy = FocalLengthInfo.FxFy.Y;
-
-				FocalLengthParametersTable.Data.Add(Focus);
-				FocalLengthParametersTable.Data.Add(Zoom);
-				FocalLengthParametersTable.Data.Add(Fx);
-				FocalLengthParametersTable.Data.Add(Fy);
-			}
-		}
-	}
-
-	CameraParameterTables.Add(FocalLengthParametersTable);
-}
-
-void FLensFileExchange::ExtractImageCenterTable(const ULensFile* LensFile)
-{
-	// Read the Image Center parameter table
-	FLensFileParameterTable ImageCenterParametersTable;
-	ImageCenterParametersTable.ParameterName = GET_MEMBER_NAME_CHECKED(ULensFile, ImageCenterTable);
-	ImageCenterParametersTable.Header = { FocusEncoderHeaderName, ZoomEncoderHeaderName, ImageCenterCxHeaderName, ImageCenterCyHeaderName };
-
-	const FImageCenterTable& LensFileImageCenterTable = LensFile->ImageCenterTable;
-
-	for (int32 ImageCenterFocusPointIndex = 0; ImageCenterFocusPointIndex < LensFileImageCenterTable.GetFocusPointNum(); ++ImageCenterFocusPointIndex)
-	{
-		const FImageCenterFocusPoint& ImageCenterFocusPoint = LensFileImageCenterTable.FocusPoints[ImageCenterFocusPointIndex];
-		const float Focus = ImageCenterFocusPoint.GetFocus();
-
-		for (int32 ImageCenterInfoIndex = 0; ImageCenterInfoIndex < ImageCenterFocusPoint.GetNumPoints(); ++ImageCenterInfoIndex)
-		{
-			const float Zoom = ImageCenterFocusPoint.GetZoom(ImageCenterInfoIndex);
-
-			FImageCenterInfo ImageCenterInfo;
-			if (ImageCenterFocusPoint.GetPoint(Zoom, ImageCenterInfo))
-			{
-				const float Cx = ImageCenterInfo.PrincipalPoint.X;
-				const float Cy = ImageCenterInfo.PrincipalPoint.Y;
-
-				ImageCenterParametersTable.Data.Add(Focus);
-				ImageCenterParametersTable.Data.Add(Zoom);
-				ImageCenterParametersTable.Data.Add(Cx);
-				ImageCenterParametersTable.Data.Add(Cy);
-			}
-		}
-	}
-
-	CameraParameterTables.Add(ImageCenterParametersTable);
-}
-
-void FLensFileExchange::ExtractNodalOffsetTable(const ULensFile* LensFile)
-{
-	// Read the NodalOffset parameter table
-	FLensFileParameterTable NodalOffsetParametersTable;
-	NodalOffsetParametersTable.ParameterName = GET_MEMBER_NAME_CHECKED(ULensFile, NodalOffsetTable);
-	NodalOffsetParametersTable.Header = {
-		FocusEncoderHeaderName,
-		ZoomEncoderHeaderName,
-		NodalOffsetQxHeaderName,
-		NodalOffsetQyHeaderName,
-		NodalOffsetQzHeaderName,
-		NodalOffsetQwHeaderName,
-		NodalOffsetTxHeaderName,
-		NodalOffsetTyHeaderName,
-		NodalOffsetTzHeaderName
-	};
-
-	const FNodalOffsetTable& LensFileNodalOffsetTable = LensFile->NodalOffsetTable;
-
-	for (int32 NodalOffsetFocusPointIndex = 0; NodalOffsetFocusPointIndex < LensFileNodalOffsetTable.GetFocusPointNum(); ++NodalOffsetFocusPointIndex)
-	{
-		const FNodalOffsetFocusPoint& NodalOffsetFocusPoint = LensFileNodalOffsetTable.FocusPoints[NodalOffsetFocusPointIndex];
-		const float Focus = NodalOffsetFocusPoint.GetFocus();
-
-		for (int32 NodalOffsetPointIndex = 0; NodalOffsetPointIndex < NodalOffsetFocusPoint.GetNumPoints(); ++NodalOffsetPointIndex)
-		{
-			const float Zoom = NodalOffsetFocusPoint.GetZoom(NodalOffsetPointIndex);
-
-			FNodalPointOffset NodalOffsetInfo;
-			if (NodalOffsetFocusPoint.GetPoint(Zoom, NodalOffsetInfo))
-			{
-				NodalOffsetParametersTable.Data.Add(Focus);
-				NodalOffsetParametersTable.Data.Add(Zoom);
-
-				FTransform NodalOffsetTransform{ NodalOffsetInfo.RotationOffset, NodalOffsetInfo.LocationOffset };
-				FCameraCalibrationUtils::ConvertUnrealToOpenCV(NodalOffsetTransform);
-
-				const FQuat RotationOffset = NodalOffsetTransform.GetRotation();
-				const FVector LocationOffset = NodalOffsetTransform.GetTranslation();
-
-				NodalOffsetParametersTable.Data.Add(RotationOffset.X);
-				NodalOffsetParametersTable.Data.Add(RotationOffset.Y);
-				NodalOffsetParametersTable.Data.Add(RotationOffset.Z);
-				NodalOffsetParametersTable.Data.Add(RotationOffset.W);
-				NodalOffsetParametersTable.Data.Add(LocationOffset.X);
-				NodalOffsetParametersTable.Data.Add(LocationOffset.Y);
-				NodalOffsetParametersTable.Data.Add(LocationOffset.Z);
-			}
-		}
-
-		CameraParameterTables.Add(NodalOffsetParametersTable);
-	}
-}
-
-void FLensFileExchange::ExtractEncoderTables(const ULensFile* LensFile)
-{
-	// Read the Encoder Tables from the LensFile
-	FLensFileParameterTable FocusEncordersTable;
-	FLensFileParameterTable IrisEncodersTable;
-
-	FocusEncordersTable.ParameterName = GET_MEMBER_NAME_CHECKED(FEncodersTable, Focus);
-	FocusEncordersTable.Header = { FocusEncoderHeaderName, FocusCMHeaderName };
-
-	IrisEncodersTable.ParameterName = GET_MEMBER_NAME_CHECKED(FEncodersTable, Iris);
-	IrisEncodersTable.Header = { IrisEncoderHeaderName, IrisFstopHeaderName };
-
-	const FEncodersTable& LensFileEncodersTable = LensFile->EncodersTable;
-
-	FocusEncordersTable.Data.Reserve(LensFileEncodersTable.GetNumFocusPoints() * 2);
-	IrisEncodersTable.Data.Reserve(LensFileEncodersTable.GetNumIrisPoints() * 2);
-
-	for (int32 FocusPointIndex = 0; FocusPointIndex < LensFileEncodersTable.GetNumFocusPoints(); ++FocusPointIndex)
-	{
-		const float FocusInput = LensFileEncodersTable.GetFocusInput(FocusPointIndex);
-		const float FocusValue = LensFileEncodersTable.GetFocusValue(FocusPointIndex);
-
-		FocusEncordersTable.Data.Add(FocusInput);
-		FocusEncordersTable.Data.Add(FocusValue);
-	}
-
-	for (int32 IrisPointIndex = 0; IrisPointIndex < LensFileEncodersTable.GetNumIrisPoints(); ++IrisPointIndex)
-	{
-		const float IrisInput = LensFileEncodersTable.GetIrisInput(IrisPointIndex);
-		const float IrisValue = LensFileEncodersTable.GetIrisValue(IrisPointIndex);
-
-		IrisEncodersTable.Data.Add(IrisInput);
-		IrisEncodersTable.Data.Add(IrisValue);
-	}
-
-	EncoderTables.Add(FocusEncordersTable);
-	EncoderTables.Add(IrisEncodersTable);
-}
-
-void FLensFileExchange::ExtractDistortionParameters(const ULensFile* LensFile)
-{
-	// Read the Distortion parameter table
-	FLensFileParameterTable DistortionParametersTable;
-	DistortionParametersTable.ParameterName = GET_MEMBER_NAME_CHECKED(ULensFile, DistortionTable);
-
-	if (const ULensModel* LensModelObject = LensFile->LensInfo.LensModel->GetDefaultObject<ULensModel>())
-	{
-		const int32 NumDistortionParams = LensModelObject->GetNumParameters();
-		DistortionParametersTable.Header.Reserve(NumDistortionParams + 2);
-
-		// Add the Focus and Zoom encoders headers that are common to all camera parameters tables
-		DistortionParametersTable.Header.Add(FocusEncoderHeaderName);
-		DistortionParametersTable.Header.Add(ZoomEncoderHeaderName);
-
-		for (const FText& ParamName : LensModelObject->GetParameterDisplayNames())
-		{
-			DistortionParametersTable.Header.Add(*ParamName.ToString());
-		}
-
-		const FDistortionTable& LensFileDistortionTable = LensFile->DistortionTable;
-		DistortionParametersTable.Data.Reserve(LensFileDistortionTable.GetFocusPointNum());
-
-		for (int32 DistortionFocusPointIndex = 0; DistortionFocusPointIndex < LensFileDistortionTable.GetFocusPointNum(); ++DistortionFocusPointIndex)
-		{
-			const FDistortionFocusPoint& DistortionFocusPoint = LensFileDistortionTable.FocusPoints[DistortionFocusPointIndex];
-			const float Focus = DistortionFocusPoint.GetFocus();
-
-			for (int32 DistortionPointIndex = 0; DistortionPointIndex < DistortionFocusPoint.GetNumPoints(); ++DistortionPointIndex)
-			{
-				const float Zoom = DistortionFocusPoint.GetZoom(DistortionPointIndex);
-
-				FDistortionInfo DistortionInfo;
-				if (DistortionFocusPoint.GetPoint(Zoom, DistortionInfo))
-				{
-					if (DistortionInfo.Parameters.Num() == NumDistortionParams)
-					{
-						// Copy the parameters to the exchange struct
-						DistortionParametersTable.Data.Add(Focus);
-						DistortionParametersTable.Data.Add(Zoom);
-
-						for (const float ParamValue : DistortionInfo.Parameters)
-						{
-							DistortionParametersTable.Data.Add(ParamValue);
-						}
-					}
-					else
-					{
-						UE_LOG(LogCameraCalibrationEditor, Warning, TEXT("Expected %d parameters for zoom %f but %d are available."), NumDistortionParams, Zoom, DistortionInfo.Parameters.Num());
-					}
-				}
-			}
-		}
-	}
-
-	CameraParameterTables.Add(DistortionParametersTable);
-}
-
-void FLensFileExchange::ExtractSTMaps(const ULensFile* LensFile)
-{
-	if (LensFile->STMapTable.GetFocusPointNum() > 0)
-	{
-		// TODO: Extract the STMaps from the LensFile to create a new FLensFileParameterTable
-		UE_LOG(LogCameraCalibrationEditor, Warning, TEXT("STMaps are not yet supported when exporting a LensFile"));
-	}
-}
-
 ULensFileExporter::ULensFileExporter(const FObjectInitializer& ObjectInitializer)
-	: Super(ObjectInitializer)
+	: Super{ ObjectInitializer }
 {
 	SupportedClass = ULensFile::StaticClass();
 	bText = true;
@@ -359,14 +59,25 @@ bool ULensFileExporter::ExportText(const class FExportObjectInnerContext* Contex
 				}
 
 				FString DataStr;
-				const int32 NumDataElements = ParameterTable->Data.Num();
-				for (int32 DataIndex = 0; DataIndex < NumDataElements; ++DataIndex)
+				const int32 NumDataRows = ParameterTable->Data.Num();
+				for (int32 DataRowIndex = 0; DataRowIndex < NumDataRows; ++DataRowIndex)
 				{
-					const float DataElement = ParameterTable->Data[DataIndex];
-					DataStr.Append(FString::SanitizeFloat(DataElement));
-					if (DataIndex < (NumDataElements - 1))
+					const FLensFileParameterTableRow& DataRow = ParameterTable->Data[DataRowIndex];
+					const int32 NumDataElements = DataRow.Values.Num();
+
+					for (int32 DataIndex = 0; DataIndex < NumDataElements; ++DataIndex)
+					{
+						const float DataElement = DataRow.Values[DataIndex];
+						DataStr.Append(FString::SanitizeFloat(DataElement));
+						if (DataIndex < (NumDataElements - 1))
+						{
+							DataStr.Append(TEXT(", "));
+						}
+					}
+
+					if (DataRowIndex < (NumDataRows - 1))
 					{
-						DataStr.Append(((DataIndex + 1) % NumHeaders == 0) ? TEXT("; ") : TEXT(", "));
+						DataStr.Append(TEXT("; "));
 					}
 				}
 
diff --git a/Engine/Plugins/VirtualProduction/CameraCalibration/Source/CameraCalibrationEditor/Private/Exporters/LensFileExporter.h b/Engine/Plugins/VirtualProduction/CameraCalibration/Source/CameraCalibrationEditor/Private/Exporters/LensFileExporter.h
--- a/Engine/Plugins/VirtualProduction/CameraCalibration/Source/CameraCalibrationEditor/Private/Exporters/LensFileExporter.h
+++ b/Engine/Plugins/VirtualProduction/CameraCalibration/Source/CameraCalibrationEditor/Private/Exporters/LensFileExporter.h
@@ -8,179 +8,6 @@
 #include "LensFileExporter.generated.h"
 
 
-// TODO: Move the Structs to another file
-
-UENUM()
-enum class ELensFileUnit
-{
-	Millimeters,
-	Normalized
-};
-
-UENUM()
-enum class ENodalOffsetCoordinateSystem
-{
-	OpenCV
-};
-
-USTRUCT()
-struct FLensInfoExchange
-{
-	GENERATED_BODY()
-
-	FLensInfoExchange(const class ULensFile* LensFile = nullptr);
-
-	UPROPERTY()
-	FName SerialNumber;
-
-	UPROPERTY()
-	FName ModelName;
-
-	UPROPERTY()
-	FName DistortionModel;
-};
-
-USTRUCT()
-struct FLensFileUserMetadataEntry
-{
-	GENERATED_BODY()
-
-	UPROPERTY()
-	FName Name;
-
-	UPROPERTY()
-	FName Value;
-};
-
-USTRUCT()
-struct FLensFileMetadata
-{
-	GENERATED_BODY()
-
-	FLensFileMetadata(const class ULensFile* LensFile = nullptr);
-
-	UPROPERTY()
-	FName Type;
-
-	UPROPERTY()
-	FName Version;
-
-	UPROPERTY()
-	FLensInfoExchange LensInfo;
-
-	UPROPERTY()
-	FName Name;
-
-	UPROPERTY()
-	ENodalOffsetCoordinateSystem NodalOffsetCoordinateSystem = ENodalOffsetCoordinateSystem::OpenCV;
-
-	UPROPERTY()
-	ELensFileUnit FxFyUnits = ELensFileUnit::Normalized;
-
-	UPROPERTY()
-	ELensFileUnit CxCyUnits = ELensFileUnit::Normalized;
-
-	UPROPERTY()
-	TArray<FLensFileUserMetadataEntry> UserMetadata;
-};
-
-
-USTRUCT()
-struct FLensFileSensorDimensions
-{
-	GENERATED_BODY()
-
-	UPROPERTY()
-	float Width = 0.0f;
-
-	UPROPERTY()
-	float Height = 0.0f;
-
-	UPROPERTY()
-	ELensFileUnit Units = ELensFileUnit::Normalized;
-};
-
-USTRUCT()
-struct FLensFileImageDimensions
-{
-	GENERATED_BODY()
-
-	UPROPERTY()
-	int32 Width = 0;
-
-	UPROPERTY()
-	int32 Height = 0;
-};
-
-USTRUCT()
-struct FLensFileParameterTable
-{
-	GENERATED_BODY()
-
-	UPROPERTY()
-	FName ParameterName;
-
-	UPROPERTY()
-	TArray<FName> Header;
-
-	UPROPERTY()
-	TArray<float> Data;
-};
-
-USTRUCT()
-struct FLensFileExchange
-{
-	GENERATED_BODY()
-
-	FLensFileExchange(const class ULensFile* LensFile = nullptr);
-
-	UPROPERTY()
-	FLensFileMetadata Metadata;
-
-	UPROPERTY()
-	FLensFileSensorDimensions SensorDimensions;
-
-	UPROPERTY()
-	FLensFileImageDimensions ImageDimensions;
-
-	UPROPERTY()
-	TArray<FLensFileParameterTable> CameraParameterTables;
-
-	UPROPERTY()
-	TArray<FLensFileParameterTable> EncoderTables;
-
-	constexpr static TCHAR LensFileVersion[] = TEXT("0.0.0");
-
-	constexpr static TCHAR FocusEncoderHeaderName[] = TEXT("FocusEncoder");
-	constexpr static TCHAR ZoomEncoderHeaderName[] = TEXT("ZoomEncoder");
-	constexpr static TCHAR IrisEncoderHeaderName[] = TEXT("IrisEncoder");
-
-	constexpr static TCHAR FocusCMHeaderName[] = TEXT("FocusCM");
-	constexpr static TCHAR IrisFstopHeaderName[] = TEXT("IrisFstop");
-
-	constexpr static TCHAR FocalLengthFxHeaderName[] = TEXT("Fx");
-	constexpr static TCHAR FocalLengthFyHeaderName[] = TEXT("Fy");
-
-	constexpr static TCHAR ImageCenterCxHeaderName[] = TEXT("Cx");
-	constexpr static TCHAR ImageCenterCyHeaderName[] = TEXT("Cy");
-
-	constexpr static TCHAR NodalOffsetQxHeaderName[] = TEXT("Qx");
-	constexpr static TCHAR NodalOffsetQyHeaderName[] = TEXT("Qy");
-	constexpr static TCHAR NodalOffsetQzHeaderName[] = TEXT("Qz");
-	constexpr static TCHAR NodalOffsetQwHeaderName[] = TEXT("Qw");
-
-	constexpr static TCHAR NodalOffsetTxHeaderName[] = TEXT("Tx");
-	constexpr static TCHAR NodalOffsetTyHeaderName[] = TEXT("Ty");
-	constexpr static TCHAR NodalOffsetTzHeaderName[] = TEXT("Tz");
-
-private:
-	void ExtractFocalLengthTable(const class ULensFile* LensFile);
-	void ExtractImageCenterTable(const class ULensFile* LensFile);
-	void ExtractNodalOffsetTable(const class ULensFile* LensFile);
-	void ExtractEncoderTables(const class ULensFile* LensFile);
-	void ExtractDistortionParameters(const class ULensFile* LensFile);
-	void ExtractSTMaps(const class ULensFile* LensFile);
-};
 
 UCLASS()
 class ULensFileExporter : public UExporter
diff --git a/Engine/Plugins/VirtualProduction/CameraCalibration/Source/CameraCalibrationEditor/Private/Factories/LensFileImporterFactory.cpp b/Engine/Plugins/VirtualProduction/CameraCalibration/Source/CameraCalibrationEditor/Private/Factories/LensFileImporterFactory.cpp
new file mode 100644
--- /dev/null
+++ b/Engine/Plugins/VirtualProduction/CameraCalibration/Source/CameraCalibrationEditor/Private/Factories/LensFileImporterFactory.cpp
@@ -0,0 +1,146 @@
+// Copyright Epic Games, Inc. All Rights Reserved.
+
+#include "LensFileImporterFactory.h"
+
+#include "CoreMinimal.h"
+#include "EditorFramework/AssetImportData.h"
+#include "JsonObjectConverter.h"
+#include "LensFile.h"
+#include "LensFileExchangeFormat.h"
+#include "Misc/FileHelper.h"
+
+
+#define LOCTEXT_NAMESPACE "LensFile Importer"
+
+ULensFileImporterFactory::ULensFileImporterFactory(const FObjectInitializer& ObjectInitializer)
+	: Super{ ObjectInitializer }
+{
+	bCreateNew = false;
+	bEditAfterNew = false;
+	SupportedClass = ULensFile::StaticClass();
+	bEditorImport = true;
+	bText = true;
+
+	Formats.Add(TEXT("ulens;Unreal LensFile"));
+}
+
+FText ULensFileImporterFactory::GetToolTip() const
+{
+	return LOCTEXT("LensFileImporterFactoryDescription", "LensFile conformant with the ulens specification");
+}
+
+bool ULensFileImporterFactory::FactoryCanImport(const FString& Filename)
+{
+	const FString Extension = FPaths::GetExtension(Filename);
+	return Extension.Compare(TEXT("ulens"), ESearchCase::IgnoreCase) == 0;
+}
+
+UObject* ULensFileImporterFactory::FactoryCreateText(UClass* InClass, UObject* InParent, FName InName, EObjectFlags Flags, UObject* Context, const TCHAR* Type, const TCHAR*& Buffer, const TCHAR* BufferEnd, FFeedbackContext* Warn)
+{
+	GEditor->GetEditorSubsystem<UImportSubsystem>()->BroadcastAssetPreImport(this, InClass, InParent, InName, Type);
+
+	ULensFile* LensFile = nullptr;
+
+	FJsonObjectWrapper JsonWrapper;
+	if (JsonWrapper.JsonObjectFromString(Buffer))
+	{
+		bool bLoadedSuccesfully = true;
+
+		const FLensFileExchange LensFileExchange{ JsonWrapper.JsonObject.ToSharedRef(), bLoadedSuccesfully, Warn };
+
+		if (bLoadedSuccesfully)
+		{
+			LensFile = NewObject<ULensFile>(InParent, InName, Flags | RF_Transactional);
+			LensFileExchange.PopulateLensFile(*LensFile);
+
+			if (LensFile->AssetImportData != nullptr)
+			{
+				LensFile->AssetImportData->Update(GetCurrentFilename());
+			}
+		}
+	}
+
+	GEditor->GetEditorSubsystem<UImportSubsystem>()->BroadcastAssetPostImport(this, LensFile);
+
+	return LensFile;
+}
+
+bool ULensFileImporterFactory::CanReimport(UObject* Obj, TArray<FString>& OutFilenames)
+{
+	if (ULensFile* LensFile = Cast<ULensFile>(Obj))
+	{
+		if (LensFile->AssetImportData != nullptr)
+		{
+			LensFile->AssetImportData->ExtractFilenames(OutFilenames);
+		}
+		else
+		{
+			OutFilenames.Add(FString{});
+		}
+
+		return true;
+	}
+	else
+	{
+		return false;
+	}
+}
+
+void ULensFileImporterFactory::SetReimportPaths(UObject* Obj, const TArray<FString>& NewReimportPaths)
+{
+	if (NewReimportPaths.IsEmpty())
+	{
+		return;
+	}
+
+	if (ULensFile* LensFile = Cast<ULensFile>(Obj))
+	{
+		if (FactoryCanImport(NewReimportPaths[0]))
+		{
+			LensFile->AssetImportData->UpdateFilenameOnly(NewReimportPaths[0]);
+		}
+	}
+}
+
+EReimportResult::Type ULensFileImporterFactory::Reimport(UObject* Obj)
+{
+	if (ULensFile* LensFile = Cast<ULensFile>(Obj))
+	{
+		// Make sure the file is valid and exists
+		const FString Filename = LensFile->AssetImportData->GetFirstFilename();
+		if (Filename.Len() == 0 || IFileManager::Get().FileSize(*Filename) == INDEX_NONE)
+		{
+			return EReimportResult::Failed;
+		}
+
+		bool bImportCancelled = false;
+		if (ImportObject(LensFile->GetClass(), LensFile->GetOuter(), *LensFile->GetName(), RF_Public | RF_Standalone, Filename, nullptr, bImportCancelled))
+		{
+			LensFile->AssetImportData->Update(Filename);
+
+			if (LensFile->GetOuter())
+			{
+				LensFile->GetOuter()->MarkPackageDirty();
+			}
+
+			LensFile->MarkPackageDirty();
+
+			return EReimportResult::Succeeded;
+		}
+		else
+		{
+			if (bImportCancelled)
+			{
+				return EReimportResult::Cancelled;
+			}
+			else
+			{
+				return EReimportResult::Failed;
+			}
+		}
+	}
+
+	return EReimportResult::Failed;
+}
+
+#undef LOCTEXT_NAMESPACE
\ No newline at end of file
diff --git a/Engine/Plugins/VirtualProduction/CameraCalibration/Source/CameraCalibrationEditor/Private/Factories/LensFileImporterFactory.h b/Engine/Plugins/VirtualProduction/CameraCalibration/Source/CameraCalibrationEditor/Private/Factories/LensFileImporterFactory.h
new file mode 100644
--- /dev/null
+++ b/Engine/Plugins/VirtualProduction/CameraCalibration/Source/CameraCalibrationEditor/Private/Factories/LensFileImporterFactory.h
@@ -0,0 +1,38 @@
+// Copyright Epic Games, Inc. All Rights Reserved.
+
+#pragma once
+
+#include "EditorReimportHandler.h"
+#include "Factories/Factory.h"
+
+#include "LensFileImporterFactory.generated.h"
+
+UCLASS()
+class ULensFileImporterFactory
+	: public UFactory
+	, public FReimportHandler
+{
+	GENERATED_BODY()
+
+public:
+	ULensFileImporterFactory(const FObjectInitializer& ObjectInitializer);
+
+protected:
+
+	//~ Begin UFactory Interface
+	virtual FText GetToolTip() const override;
+	virtual bool FactoryCanImport(const FString& Filename) override;
+	virtual UObject* FactoryCreateText(UClass* InClass, UObject* InParent, FName InName, EObjectFlags Flags, UObject* Context, const TCHAR* Type, const TCHAR*& Buffer, const TCHAR* BufferEnd, FFeedbackContext* Warn) override;
+	//~ End UFactory Interface
+
+	//~ Begin FReimportHandler Interface
+	virtual bool CanReimport(UObject* Obj, TArray<FString>& OutFilenames) override;
+	virtual void SetReimportPaths(UObject* Obj, const TArray<FString>& NewReimportPaths) override;
+	virtual EReimportResult::Type Reimport(UObject* Obj) override;
+
+	virtual const UObject* GetFactoryObject() const 
+	{
+		return this;
+	}
+	//~ End FReimiporterHandler Interface
+};
\ No newline at end of file
diff --git a/Engine/Plugins/VirtualProduction/CameraCalibration/Source/CameraCalibrationEditor/Private/LensFileExchangeFormat.cpp b/Engine/Plugins/VirtualProduction/CameraCalibration/Source/CameraCalibrationEditor/Private/LensFileExchangeFormat.cpp
new file mode 100644
--- /dev/null
+++ b/Engine/Plugins/VirtualProduction/CameraCalibration/Source/CameraCalibrationEditor/Private/LensFileExchangeFormat.cpp
@@ -0,0 +1,939 @@
+// Copyright Epic Games, Inc. All Rights Reserved.
+
+#include "LensFileExchangeFormat.h"
+
+#include "CameraCalibrationEditorLog.h"
+#include "CameraCalibrationUtils.h"
+#include "JsonObjectConverter.h"
+#include "LensData.h"
+#include "LensFile.h"
+#include "Misc/FeedbackContext.h"
+#include "Models/LensModel.h"
+#include "Models/SphericalLensModel.h"
+
+
+#define LOCTEXT_NAMESPACE "LensFileExchange"
+
+FLensInfoExchange::FLensInfoExchange(const ULensFile* LensFile)
+{
+	if (LensFile == nullptr)
+	{
+		return;
+	}
+
+	const FLensInfo& LensInfo = LensFile->LensInfo;
+
+	SerialNumber = *LensInfo.LensSerialNumber;
+	ModelName = *LensInfo.LensModelName;
+
+	if (const ULensModel* LensModelObject = LensInfo.LensModel->GetDefaultObject<ULensModel>())
+	{
+		DistortionModel = LensModelObject->GetShortModelName();
+	}
+}
+
+FLensFileMetadata::FLensFileMetadata(const ULensFile* LensFile)
+	: LensInfo{ LensFile }
+{
+	Type = FLensFileExchange::LensFileType;
+	Version = FLensFileExchange::LensFileVersion;
+
+	if (LensFile == nullptr)
+	{
+		return;
+	}
+
+	// Set the default values for metadata when exporting from the ULensFile
+	Name = *LensFile->GetName();
+
+	// Add the UserMetadata
+	UserMetadata.Reserve(LensFile->UserMetadata.Num());
+	for (const TPair<FString, FString>& UserMetadataPair : LensFile->UserMetadata)
+	{
+		FLensFileUserMetadataEntry UserMetadataEntry;
+		UserMetadataEntry.Name = *UserMetadataPair.Key;
+		UserMetadataEntry.Value = *UserMetadataPair.Value;
+		UserMetadata.Add(UserMetadataEntry);
+	}
+}
+
+FLensFileParameterTableRow::FLensFileParameterTableRow(int32 NumHeaders)
+{
+	Values.Reserve(NumHeaders);
+}
+
+bool FLensFileParameterTableImporter::ConvertToParameterTable(FLensFileParameterTable& OutParamterTable, FFeedbackContext* InWarn) const
+{
+	OutParamterTable.ParameterName = ParameterName;
+
+	// Parse the header string
+	TArray<FString> HeaderTokens;
+	Header.ParseIntoArrayWS(HeaderTokens, TEXT(","));
+
+	const int32 NumHeaders = HeaderTokens.Num();
+
+	OutParamterTable.Header.Reserve(NumHeaders);
+	for (const FString& HeaderToken : HeaderTokens)
+	{
+		OutParamterTable.Header.Add(*HeaderToken);
+	}
+
+	// Parse the data string
+	TArray<FString> DataRowTokens;
+	Data.ParseIntoArray(DataRowTokens, TEXT(";"), 1);
+
+	// Reserve memory to fill in the Data array of the parameter table
+	OutParamterTable.Data.Reserve(NumHeaders * DataRowTokens.Num());
+
+	// Split each row and convert values to floats
+	// TODO: Need to handle other data types
+	for (const FString& DataRowToken : DataRowTokens)
+	{
+		TArray<FString> DataTokens;
+		DataRowToken.ParseIntoArrayWS(DataTokens, TEXT(","));
+
+		FLensFileParameterTableRow DataRow{ NumHeaders };
+
+		for (const FString& DataToken : DataTokens)
+		{
+			float DataValue = 0.0f;
+			if (LexTryParseString(DataValue, *DataToken))
+			{
+				DataRow.Values.Add(DataValue);
+			}
+			else
+			{
+				InWarn->CategorizedLogf(LogCameraCalibrationEditor.GetCategoryName(),
+										ELogVerbosity::Error,
+										TEXT("Error converting '%s' value in table '%s'"),
+										*DataToken, *ParameterName.ToString());
+			}
+		}
+
+		OutParamterTable.Data.Add(DataRow);
+	}
+
+	return true;
+}
+
+FLensFileExchange::FLensFileExchange(const ULensFile* LensFile)
+	: Metadata{ LensFile }
+	, FeedbackContext{ nullptr }
+{
+	if (LensFile == nullptr)
+	{
+		return;
+	}
+
+	const FLensInfo& LensInfo = LensFile->LensInfo;
+	const FVector2D& Sensor = LensInfo.SensorDimensions;
+
+	SensorDimensions.Width = Sensor.X;
+	SensorDimensions.Height = Sensor.Y;
+	SensorDimensions.Units = ELensFileUnit::Millimeters;
+
+	ImageDimensions.Width = LensInfo.ImageDimensions.X;
+	ImageDimensions.Height = LensInfo.ImageDimensions.Y;
+
+	ExtractFocalLengthTable(LensFile);
+	ExtractImageCenterTable(LensFile);
+	ExtractNodalOffsetTable(LensFile);
+	ExtractEncoderTables(LensFile);
+	ExtractDistortionParameters(LensFile);
+	ExtractSTMaps(LensFile);
+}
+
+FLensFileExchange::FLensFileExchange(const TSharedRef<FJsonObject>& InLensFileJson, bool& bOutIsValidJson, FFeedbackContext* InWarn)
+	: FeedbackContext{ InWarn }
+{
+	// Start by assuming the Json is valid
+	bOutIsValidJson = true;
+
+	const int64 CheckFlags = 0;
+	const int64 SkipFlags = 0;
+	const bool bStrictMode = true;
+
+	const TSharedPtr<FJsonObject>* MetadataJsonObject;
+	const FString MetadataJsonFieldName = FJsonObjectConverter::StandardizeCase(GET_MEMBER_NAME_CHECKED(FLensFileExchange, Metadata).ToString());
+	if (InLensFileJson->TryGetObjectField(*MetadataJsonFieldName, MetadataJsonObject))
+	{
+		if (FJsonObjectConverter::JsonObjectToUStruct(MetadataJsonObject->ToSharedRef(), &Metadata, CheckFlags, SkipFlags, bStrictMode))
+		{
+			if (Metadata.Type != FLensFileExchange::LensFileType)
+			{
+				GetFeedbackContext()->CategorizedLogf(LogCameraCalibrationEditor.GetCategoryName(),
+													  ELogVerbosity::Error,
+													  TEXT("Invalid value for field 'type'. Expected '%s' but got '%s'"),
+													  FLensFileExchange::LensFileType, *Metadata.Type.ToString());
+				bOutIsValidJson = false;
+			}
+
+			if (Metadata.Version != FLensFileExchange::LensFileVersion)
+			{
+				// TODO: Add proper version validation, a string comparison is enough for now
+				GetFeedbackContext()->CategorizedLogf(LogCameraCalibrationEditor.GetCategoryName(),
+													  ELogVerbosity::Error,
+													  TEXT("Invalid value for field 'version'. Expected '%s' but got '%s'"),
+													  FLensFileExchange::LensFileVersion, *Metadata.Type.ToString());
+				bOutIsValidJson = false;
+			}
+		}
+		else
+		{
+			GetFeedbackContext()->CategorizedLogf(LogCameraCalibrationEditor.GetCategoryName(),
+												  ELogVerbosity::Error,
+												  TEXT("Error parsing '%s' field"),
+												  *MetadataJsonFieldName);
+			bOutIsValidJson = false;
+		}
+	}
+	else
+	{
+		GetFeedbackContext()->CategorizedLogf(LogCameraCalibrationEditor.GetCategoryName(),
+											  ELogVerbosity::Error, TEXT("Field '%s' not found"),
+											  *MetadataJsonFieldName);
+		bOutIsValidJson = false;
+	}
+
+	if (!bOutIsValidJson)
+	{
+		return;
+	}
+
+	const TSharedPtr<FJsonObject>* SensorDimensionsJsonObject;
+	const FString SensorDimensionsJsonFieldName = FJsonObjectConverter::StandardizeCase(GET_MEMBER_NAME_CHECKED(FLensFileExchange, SensorDimensions).ToString());
+	if (InLensFileJson->TryGetObjectField(SensorDimensionsJsonFieldName, SensorDimensionsJsonObject))
+	{
+		if (FJsonObjectConverter::JsonObjectToUStruct(SensorDimensionsJsonObject->ToSharedRef(), &SensorDimensions, CheckFlags, SkipFlags, bStrictMode))
+		{
+			if (SensorDimensions.Units != ELensFileUnit::Millimeters)
+			{
+				GetFeedbackContext()->Log(LogCameraCalibrationEditor.GetCategoryName(), ELogVerbosity::Warning, TEXT("Invalid unit for sensor dimensions"));
+			}
+		}
+		else
+		{
+			GetFeedbackContext()->CategorizedLogf(LogCameraCalibrationEditor.GetCategoryName(),
+												  ELogVerbosity::Warning,
+												  TEXT("Error parsing '%s' field"),
+												  *SensorDimensionsJsonFieldName);
+		}
+	}
+	else
+	{
+		GetFeedbackContext()->CategorizedLogf(LogCameraCalibrationEditor.GetCategoryName(),
+											  ELogVerbosity::Error,
+											  TEXT("Field '%s' not found"),
+											  *SensorDimensionsJsonFieldName);
+		bOutIsValidJson = false;
+	}
+
+	if (!bOutIsValidJson)
+	{
+		return;
+	}
+
+	const TSharedPtr<FJsonObject>* ImageDimensionsJsonObject;
+	const FString ImageDimensionsJsonFieldName = FJsonObjectConverter::StandardizeCase(GET_MEMBER_NAME_CHECKED(FLensFileExchange, ImageDimensions).ToString());
+	if (InLensFileJson->TryGetObjectField(ImageDimensionsJsonFieldName, ImageDimensionsJsonObject))
+	{
+		if (!FJsonObjectConverter::JsonObjectToUStruct(ImageDimensionsJsonObject->ToSharedRef(), &ImageDimensions, CheckFlags, SkipFlags, bStrictMode))
+		{
+			GetFeedbackContext()->CategorizedLogf(LogCameraCalibrationEditor.GetCategoryName(),
+												  ELogVerbosity::Error,
+												  TEXT("Error parsing '%s' field"),
+												  *ImageDimensionsJsonFieldName);
+			bOutIsValidJson = false;
+		}
+	}
+	else
+	{
+		GetFeedbackContext()->CategorizedLogf(LogCameraCalibrationEditor.GetCategoryName(), ELogVerbosity::Error, TEXT("Field '%s' not found"), *ImageDimensionsJsonFieldName);
+		bOutIsValidJson = false;
+	}
+
+	if (!bOutIsValidJson)
+	{
+		return;
+	}
+
+	const TArray<TSharedPtr<FJsonValue>>* CameraParameterTablesValueArray;
+	const FString CameraParameterTablesJsonFieldName = FJsonObjectConverter::StandardizeCase(GET_MEMBER_NAME_CHECKED(FLensFileExchange, CameraParameterTables).ToString());
+	if (InLensFileJson->TryGetArrayField(CameraParameterTablesJsonFieldName, CameraParameterTablesValueArray))
+	{
+		CameraParameterTables.Reserve(CameraParameterTablesValueArray->Num());
+
+		for (const TSharedPtr<FJsonValue>& CameraParameterTableObjectValue : *CameraParameterTablesValueArray)
+		{
+			const TSharedPtr<FJsonObject>* ParameterTableJsonObject;
+			if (CameraParameterTableObjectValue->TryGetObject(ParameterTableJsonObject))
+			{
+				FLensFileParameterTableImporter ParameterTableImporter;
+				if (FJsonObjectConverter::JsonObjectToUStruct(ParameterTableJsonObject->ToSharedRef(), &ParameterTableImporter, CheckFlags, SkipFlags, bStrictMode))
+				{
+					FLensFileParameterTable CameraParameterTable;
+					ParameterTableImporter.ConvertToParameterTable(CameraParameterTable, GetFeedbackContext());
+					CameraParameterTables.Add(CameraParameterTable);
+				}
+				else
+				{
+					GetFeedbackContext()->CategorizedLogf(LogCameraCalibrationEditor.GetCategoryName(),
+														  ELogVerbosity::Warning,
+														  TEXT("Error parsing array entry in '%s' field"),
+														  *CameraParameterTablesJsonFieldName);
+				}
+			}
+			else
+			{
+				GetFeedbackContext()->CategorizedLogf(LogCameraCalibrationEditor.GetCategoryName(),
+													  ELogVerbosity::Warning,
+													  TEXT("Entries in the '%s' array should be objects"),
+													  *CameraParameterTablesJsonFieldName);
+			}
+		}
+	}
+	else
+	{
+		GetFeedbackContext()->CategorizedLogf(LogCameraCalibrationEditor.GetCategoryName(),
+											  ELogVerbosity::Warning,
+											  TEXT("Array field '%s' not found"),
+											  *CameraParameterTablesJsonFieldName);
+	}
+
+	const TArray<TSharedPtr<FJsonValue>>* EncoderTablesValueArray;
+	const FString EncoderTablesJsonFieldName = FJsonObjectConverter::StandardizeCase(GET_MEMBER_NAME_CHECKED(FLensFileExchange, EncoderTables).ToString());
+	if (InLensFileJson->TryGetArrayField(EncoderTablesJsonFieldName, EncoderTablesValueArray))
+	{
+		EncoderTables.Reserve(EncoderTablesValueArray->Num());
+
+		for (const TSharedPtr<FJsonValue>& EncoderTableObjectValue : *EncoderTablesValueArray)
+		{
+			const TSharedPtr<FJsonObject>* EncoderTableJsonObject;
+			if (EncoderTableObjectValue->TryGetObject(EncoderTableJsonObject))
+			{
+				FLensFileParameterTableImporter EncoderTableImporter;
+				if (FJsonObjectConverter::JsonObjectToUStruct(EncoderTableJsonObject->ToSharedRef(), &EncoderTableImporter, CheckFlags, SkipFlags, bStrictMode))
+				{
+					FLensFileParameterTable EncoderTable;
+					EncoderTableImporter.ConvertToParameterTable(EncoderTable, GetFeedbackContext());
+					EncoderTables.Add(EncoderTable);
+				}
+				else
+				{
+					GetFeedbackContext()->CategorizedLogf(LogCameraCalibrationEditor.GetCategoryName(),
+														  ELogVerbosity::Warning,
+														  TEXT("Error parsing array entry in '%s' field"),
+														  *EncoderTablesJsonFieldName);
+				}
+			}
+			else
+			{
+				GetFeedbackContext()->CategorizedLogf(LogCameraCalibrationEditor.GetCategoryName(),
+													  ELogVerbosity::Warning,
+													  TEXT("Entries in the '%s' array should be objects"),
+													  *EncoderTablesJsonFieldName);
+			}
+		}
+	}
+	else
+	{
+		GetFeedbackContext()->CategorizedLogf(LogCameraCalibrationEditor.GetCategoryName(),
+											  ELogVerbosity::Warning,
+											  TEXT("Array field '%s' not found"),
+											  *EncoderTablesJsonFieldName);
+	}
+}
+
+bool FLensFileExchange::PopulateLensFile(ULensFile& OutLensFile) const
+{
+	FLensInfo& LensInfo = OutLensFile.LensInfo;
+
+	const USphericalLensModel* SphericalLensModel = USphericalLensModel::StaticClass()->GetDefaultObject<USphericalLensModel>();
+	const bool bIsSphericalLensModel = Metadata.LensInfo.DistortionModel == SphericalLensModel->GetShortModelName();
+	if (bIsSphericalLensModel)
+	{
+		LensInfo.LensModel = USphericalLensModel::StaticClass();
+	}
+	else
+	{
+		GetFeedbackContext()->CategorizedLogf(LogCameraCalibrationEditor.GetCategoryName(),
+											  ELogVerbosity::Error,
+											  TEXT("Only %s is supported at the moment"),
+											  *SphericalLensModel->GetShortModelName().ToString());
+		return false;
+	}
+
+	LensInfo.LensModelName = Metadata.LensInfo.ModelName.ToString();
+	LensInfo.LensSerialNumber = Metadata.LensInfo.SerialNumber.ToString();
+
+	// TODO: Handle unit conversions, assuming units are stored in millimeters
+	LensInfo.SensorDimensions.X = SensorDimensions.Width;
+	LensInfo.SensorDimensions.Y = SensorDimensions.Height;
+
+	LensInfo.ImageDimensions.X = ImageDimensions.Width;
+	LensInfo.ImageDimensions.Y = ImageDimensions.Height;
+
+	for (const FLensFileParameterTable& CameraParameterTable : CameraParameterTables)
+	{
+		if (CameraParameterTable.ParameterName == GET_MEMBER_NAME_CHECKED(ULensFile, FocalLengthTable))
+		{
+			PopulateFocalLengthTable(OutLensFile, CameraParameterTable);
+		}
+		else if (CameraParameterTable.ParameterName == GET_MEMBER_NAME_CHECKED(ULensFile, ImageCenterTable))
+		{
+			PopulateImageCenterTable(OutLensFile, CameraParameterTable);
+		}
+		else if (CameraParameterTable.ParameterName == GET_MEMBER_NAME_CHECKED(ULensFile, DistortionTable))
+		{
+			PopulateDistortionTable(OutLensFile, CameraParameterTable);
+		}
+		else if (CameraParameterTable.ParameterName == GET_MEMBER_NAME_CHECKED(ULensFile, NodalOffsetTable))
+		{
+			PopulateNodalOffsetTable(OutLensFile, CameraParameterTable);
+		}
+	}
+
+	for (const FLensFileParameterTable& EncoderTable : EncoderTables)
+	{
+		if (EncoderTable.ParameterName == GET_MEMBER_NAME_CHECKED(FEncodersTable, Focus))
+		{
+			PopulateEncoderTable(OutLensFile.EncodersTable.Focus, EncoderTable, FocusEncoderHeaderName, FocusCMHeaderName);
+		}
+		else if (EncoderTable.ParameterName == GET_MEMBER_NAME_CHECKED(FEncodersTable, Iris))
+		{
+			PopulateEncoderTable(OutLensFile.EncodersTable.Iris, EncoderTable, IrisEncoderHeaderName, IrisFstopHeaderName);
+		}
+	}
+
+	return true;
+}
+
+void FLensFileExchange::ExtractFocalLengthTable(const ULensFile* LensFile)
+{
+	// Read the FocalLength table
+	FLensFileParameterTable FocalLengthParametersTable;
+	FocalLengthParametersTable.ParameterName = GET_MEMBER_NAME_CHECKED(ULensFile, FocalLengthTable);
+	FocalLengthParametersTable.Header = { FocusEncoderHeaderName, ZoomEncoderHeaderName, FocalLengthFxHeaderName, FocalLengthFyHeaderName };
+
+	const FFocalLengthTable& LensFileFocalLengthTable = LensFile->FocalLengthTable;
+
+	for (int32 FocusPointIndex = 0; FocusPointIndex < LensFileFocalLengthTable.GetFocusPointNum(); ++FocusPointIndex)
+	{
+		const FFocalLengthFocusPoint& FocalLengthFocusPoint = LensFileFocalLengthTable.FocusPoints[FocusPointIndex];
+		const float Focus = FocalLengthFocusPoint.GetFocus();
+
+		for (int32 FocalLengthInfoIndex = 0; FocalLengthInfoIndex < FocalLengthFocusPoint.GetNumPoints(); ++FocalLengthInfoIndex)
+		{
+			FFocalLengthInfo FocalLengthInfo;
+			if (FocalLengthFocusPoint.GetValue(FocalLengthInfoIndex, FocalLengthInfo))
+			{
+				const float Zoom = FocalLengthFocusPoint.GetZoom(FocalLengthInfoIndex);
+				const float Fx = FocalLengthInfo.FxFy.X;
+				const float Fy = FocalLengthInfo.FxFy.Y;
+
+				FLensFileParameterTableRow DataRow{ FocalLengthParametersTable.Header.Num() };
+
+				DataRow.Values.Add(Focus);
+				DataRow.Values.Add(Zoom);
+				DataRow.Values.Add(Fx);
+				DataRow.Values.Add(Fy);
+
+				FocalLengthParametersTable.Data.Add(DataRow);
+			}
+		}
+	}
+
+	CameraParameterTables.Add(FocalLengthParametersTable);
+}
+
+void FLensFileExchange::ExtractImageCenterTable(const ULensFile* LensFile)
+{
+	// Read the Image Center parameter table
+	FLensFileParameterTable ImageCenterParametersTable;
+	ImageCenterParametersTable.ParameterName = GET_MEMBER_NAME_CHECKED(ULensFile, ImageCenterTable);
+	ImageCenterParametersTable.Header = { FocusEncoderHeaderName, ZoomEncoderHeaderName, ImageCenterCxHeaderName, ImageCenterCyHeaderName };
+
+	const FImageCenterTable& LensFileImageCenterTable = LensFile->ImageCenterTable;
+
+	for (int32 ImageCenterFocusPointIndex = 0; ImageCenterFocusPointIndex < LensFileImageCenterTable.GetFocusPointNum(); ++ImageCenterFocusPointIndex)
+	{
+		const FImageCenterFocusPoint& ImageCenterFocusPoint = LensFileImageCenterTable.FocusPoints[ImageCenterFocusPointIndex];
+		const float Focus = ImageCenterFocusPoint.GetFocus();
+
+		for (int32 ImageCenterInfoIndex = 0; ImageCenterInfoIndex < ImageCenterFocusPoint.GetNumPoints(); ++ImageCenterInfoIndex)
+		{
+			const float Zoom = ImageCenterFocusPoint.GetZoom(ImageCenterInfoIndex);
+
+			FImageCenterInfo ImageCenterInfo;
+			if (ImageCenterFocusPoint.GetPoint(Zoom, ImageCenterInfo))
+			{
+				const float Cx = ImageCenterInfo.PrincipalPoint.X;
+				const float Cy = ImageCenterInfo.PrincipalPoint.Y;
+
+				FLensFileParameterTableRow DataRow{ ImageCenterParametersTable.Header.Num() };
+
+				DataRow.Values.Add(Focus);
+				DataRow.Values.Add(Zoom);
+				DataRow.Values.Add(Cx);
+				DataRow.Values.Add(Cy);
+
+				ImageCenterParametersTable.Data.Add(DataRow);
+			}
+		}
+	}
+
+	CameraParameterTables.Add(ImageCenterParametersTable);
+}
+
+void FLensFileExchange::ExtractNodalOffsetTable(const ULensFile* LensFile)
+{
+	// Read the NodalOffset parameter table
+	FLensFileParameterTable NodalOffsetParametersTable;
+	NodalOffsetParametersTable.ParameterName = GET_MEMBER_NAME_CHECKED(ULensFile, NodalOffsetTable);
+	NodalOffsetParametersTable.Header = {
+		FocusEncoderHeaderName,
+		ZoomEncoderHeaderName,
+		NodalOffsetQxHeaderName,
+		NodalOffsetQyHeaderName,
+		NodalOffsetQzHeaderName,
+		NodalOffsetQwHeaderName,
+		NodalOffsetTxHeaderName,
+		NodalOffsetTyHeaderName,
+		NodalOffsetTzHeaderName
+	};
+
+	const FNodalOffsetTable& LensFileNodalOffsetTable = LensFile->NodalOffsetTable;
+
+	for (int32 NodalOffsetFocusPointIndex = 0; NodalOffsetFocusPointIndex < LensFileNodalOffsetTable.GetFocusPointNum(); ++NodalOffsetFocusPointIndex)
+	{
+		const FNodalOffsetFocusPoint& NodalOffsetFocusPoint = LensFileNodalOffsetTable.FocusPoints[NodalOffsetFocusPointIndex];
+		const float Focus = NodalOffsetFocusPoint.GetFocus();
+
+		for (int32 NodalOffsetPointIndex = 0; NodalOffsetPointIndex < NodalOffsetFocusPoint.GetNumPoints(); ++NodalOffsetPointIndex)
+		{
+			const float Zoom = NodalOffsetFocusPoint.GetZoom(NodalOffsetPointIndex);
+
+			FNodalPointOffset NodalOffsetInfo;
+			if (NodalOffsetFocusPoint.GetPoint(Zoom, NodalOffsetInfo))
+			{
+				FTransform NodalOffsetTransform{ NodalOffsetInfo.RotationOffset, NodalOffsetInfo.LocationOffset };
+				FCameraCalibrationUtils::ConvertUnrealToOpenCV(NodalOffsetTransform);
+
+				const FQuat RotationOffset = NodalOffsetTransform.GetRotation();
+				const FVector LocationOffset = NodalOffsetTransform.GetTranslation();
+
+				FLensFileParameterTableRow DataRow{ NodalOffsetParametersTable.Header.Num() };
+
+				DataRow.Values.Add(Focus);
+				DataRow.Values.Add(Zoom);
+				DataRow.Values.Add(RotationOffset.X);
+				DataRow.Values.Add(RotationOffset.Y);
+				DataRow.Values.Add(RotationOffset.Z);
+				DataRow.Values.Add(RotationOffset.W);
+				DataRow.Values.Add(LocationOffset.X);
+				DataRow.Values.Add(LocationOffset.Y);
+				DataRow.Values.Add(LocationOffset.Z);
+
+				NodalOffsetParametersTable.Data.Add(DataRow);
+			}
+		}
+	}
+
+	CameraParameterTables.Add(NodalOffsetParametersTable);
+}
+
+void FLensFileExchange::ExtractEncoderTables(const ULensFile* LensFile)
+{
+	// Read the Encoder Tables from the LensFile
+	FLensFileParameterTable FocusEncordersTable;
+	FLensFileParameterTable IrisEncodersTable;
+
+	FocusEncordersTable.ParameterName = GET_MEMBER_NAME_CHECKED(FEncodersTable, Focus);
+	FocusEncordersTable.Header = { FocusEncoderHeaderName, FocusCMHeaderName };
+
+	IrisEncodersTable.ParameterName = GET_MEMBER_NAME_CHECKED(FEncodersTable, Iris);
+	IrisEncodersTable.Header = { IrisEncoderHeaderName, IrisFstopHeaderName };
+
+	const FEncodersTable& LensFileEncodersTable = LensFile->EncodersTable;
+
+	FocusEncordersTable.Data.Reserve(LensFileEncodersTable.GetNumFocusPoints() * 2);
+	IrisEncodersTable.Data.Reserve(LensFileEncodersTable.GetNumIrisPoints() * 2);
+
+	for (int32 FocusPointIndex = 0; FocusPointIndex < LensFileEncodersTable.GetNumFocusPoints(); ++FocusPointIndex)
+	{
+		const float FocusInput = LensFileEncodersTable.GetFocusInput(FocusPointIndex);
+		const float FocusValue = LensFileEncodersTable.GetFocusValue(FocusPointIndex);
+
+		FLensFileParameterTableRow DataRow{ FocusEncordersTable.Header.Num() };
+
+		DataRow.Values.Add(FocusInput);
+		DataRow.Values.Add(FocusValue);
+
+		FocusEncordersTable.Data.Add(DataRow);
+	}
+
+	for (int32 IrisPointIndex = 0; IrisPointIndex < LensFileEncodersTable.GetNumIrisPoints(); ++IrisPointIndex)
+	{
+		const float IrisInput = LensFileEncodersTable.GetIrisInput(IrisPointIndex);
+		const float IrisValue = LensFileEncodersTable.GetIrisValue(IrisPointIndex);
+
+		FLensFileParameterTableRow DataRow{ IrisEncodersTable.Header.Num() };
+
+		DataRow.Values.Add(IrisInput);
+		DataRow.Values.Add(IrisValue);
+
+		IrisEncodersTable.Data.Add(DataRow);
+	}
+
+	EncoderTables.Add(FocusEncordersTable);
+	EncoderTables.Add(IrisEncodersTable);
+}
+
+void FLensFileExchange::ExtractDistortionParameters(const ULensFile* LensFile)
+{
+	// Read the Distortion parameter table
+	FLensFileParameterTable DistortionParametersTable;
+	DistortionParametersTable.ParameterName = GET_MEMBER_NAME_CHECKED(ULensFile, DistortionTable);
+
+	if (const ULensModel* LensModelObject = LensFile->LensInfo.LensModel->GetDefaultObject<ULensModel>())
+	{
+		const int32 NumDistortionParams = LensModelObject->GetNumParameters();
+		DistortionParametersTable.Header.Reserve(NumDistortionParams + 2);
+
+		// Add the Focus and Zoom encoders headers that are common to all camera parameters tables
+		DistortionParametersTable.Header.Add(FocusEncoderHeaderName);
+		DistortionParametersTable.Header.Add(ZoomEncoderHeaderName);
+
+		for (const FText& ParamName : LensModelObject->GetParameterDisplayNames())
+		{
+			DistortionParametersTable.Header.Add(*ParamName.ToString());
+		}
+
+		const FDistortionTable& LensFileDistortionTable = LensFile->DistortionTable;
+		DistortionParametersTable.Data.Reserve(LensFileDistortionTable.GetFocusPointNum());
+
+		for (int32 DistortionFocusPointIndex = 0; DistortionFocusPointIndex < LensFileDistortionTable.GetFocusPointNum(); ++DistortionFocusPointIndex)
+		{
+			const FDistortionFocusPoint& DistortionFocusPoint = LensFileDistortionTable.FocusPoints[DistortionFocusPointIndex];
+			const float Focus = DistortionFocusPoint.GetFocus();
+
+			for (int32 DistortionPointIndex = 0; DistortionPointIndex < DistortionFocusPoint.GetNumPoints(); ++DistortionPointIndex)
+			{
+				const float Zoom = DistortionFocusPoint.GetZoom(DistortionPointIndex);
+
+				FDistortionInfo DistortionInfo;
+				if (DistortionFocusPoint.GetPoint(Zoom, DistortionInfo))
+				{
+					if (DistortionInfo.Parameters.Num() == NumDistortionParams)
+					{
+						FLensFileParameterTableRow DataRow{ DistortionParametersTable.Header.Num() };
+
+						// Copy the parameters to the exchange struct
+						DataRow.Values.Add(Focus);
+						DataRow.Values.Add(Zoom);
+
+						for (const float ParamValue : DistortionInfo.Parameters)
+						{
+							DataRow.Values.Add(ParamValue);
+						}
+
+						DistortionParametersTable.Data.Add(DataRow);
+					}
+					else
+					{
+						UE_LOG(LogCameraCalibrationEditor, Warning, TEXT("Expected %d parameters for zoom %f but only %d are available."), NumDistortionParams, Zoom, DistortionInfo.Parameters.Num());
+					}
+				}
+			}
+		}
+	}
+
+	CameraParameterTables.Add(DistortionParametersTable);
+}
+
+void FLensFileExchange::ExtractSTMaps(const ULensFile* LensFile)
+{
+	if (LensFile->STMapTable.GetFocusPointNum() > 0)
+	{
+		// TODO: Extract the STMaps from the LensFile to create a new FLensFileParameterTable
+		UE_LOG(LogCameraCalibrationEditor, Warning, TEXT("STMaps are not yet supported when exporting a LensFile"));
+	}
+}
+
+void FLensFileExchange::PopulateFocalLengthTable(ULensFile& OutLensFile, const FLensFileParameterTable& InFocalLengthTable) const
+{
+	const int32 NumHeaders = InFocalLengthTable.Header.Num();
+
+	for (int32 DataRowIndex = 0; DataRowIndex < InFocalLengthTable.Data.Num(); ++DataRowIndex)
+	{
+		const FLensFileParameterTableRow& DataRow = InFocalLengthTable.Data[DataRowIndex];
+		const int32 NumDataValues = DataRow.Values.Num();
+
+		if (NumDataValues == NumHeaders)
+		{
+			float Focus = 0.0f;
+			float Zoom = 0.0f;
+			FVector2D FxFy;
+
+			for (int32 HeaderIndex = 0; HeaderIndex < NumHeaders; ++HeaderIndex)
+			{
+				const FName& HeaderName = InFocalLengthTable.Header[HeaderIndex];
+				const float Value = DataRow.Values[HeaderIndex];
+
+				if (HeaderName == FocusEncoderHeaderName)
+				{
+					Focus = Value;
+				}
+				else if (HeaderName == ZoomEncoderHeaderName)
+				{
+					Zoom = Value;
+				}
+				else if (HeaderName == FocalLengthFxHeaderName)
+				{
+					FxFy.X = Value;
+				}
+				else if (HeaderName == FocalLengthFyHeaderName)
+				{
+					FxFy.Y = Value;
+				}
+			}
+
+			OutLensFile.AddFocalLengthPoint(Focus, Zoom, FFocalLengthInfo{ FxFy });
+		}
+		else
+		{
+			GetFeedbackContext()->CategorizedLogf(LogCameraCalibrationEditor.GetCategoryName(),
+												  ELogVerbosity::Warning,
+												  TEXT("Number of values in row %d of table '%s' doesn't match with the number of headers. Expected %d values but got %d. Skipping line."),
+												  DataRowIndex, *InFocalLengthTable.ParameterName.ToString(), NumHeaders, NumDataValues);
+		}
+	}
+}
+
+void FLensFileExchange::PopulateImageCenterTable(ULensFile& OutLensFile, const FLensFileParameterTable& InImageCenterTable) const
+{
+	const int32 NumHeaders = InImageCenterTable.Header.Num();
+
+	for (int32 DataRowIndex = 0; DataRowIndex < InImageCenterTable.Data.Num(); ++DataRowIndex)
+	{
+		const FLensFileParameterTableRow& DataRow = InImageCenterTable.Data[DataRowIndex];
+		const int32 NumDataValues = DataRow.Values.Num();
+
+		if (NumDataValues == NumHeaders)
+		{
+			float Focus = 0.0f;
+			float Zoom = 0.0f;
+			FVector2D FxFy;
+
+			for (int32 HeaderIndex = 0; HeaderIndex < NumHeaders; ++HeaderIndex)
+			{
+				const FName& HeaderName = InImageCenterTable.Header[HeaderIndex];
+				const float Value = DataRow.Values[HeaderIndex];
+
+				if (HeaderName == FocusEncoderHeaderName)
+				{
+					Focus = Value;
+				}
+				else if (HeaderName == ZoomEncoderHeaderName)
+				{
+					Zoom = Value;
+				}
+				else if (HeaderName == ImageCenterCxHeaderName)
+				{
+					FxFy.X = Value;
+				}
+				else if (HeaderName == ImageCenterCyHeaderName)
+				{
+					FxFy.Y = Value;
+				}
+			}
+
+			OutLensFile.AddImageCenterPoint(Focus, Zoom, FImageCenterInfo{ FxFy });
+		}
+		else
+		{
+			GetFeedbackContext()->CategorizedLogf(LogCameraCalibrationEditor.GetCategoryName(),
+												  ELogVerbosity::Warning,
+												  TEXT("Number of values in row %d of table '%s' doesn't match with the number of headers. Expected %d values but got %d. Skipping line."),
+												  DataRowIndex, *InImageCenterTable.ParameterName.ToString(), NumHeaders, NumDataValues);
+		}
+	}
+}
+
+void FLensFileExchange::PopulateNodalOffsetTable(ULensFile& OutLensFile, const FLensFileParameterTable& InNodalOffsetTable) const
+{
+	const int32 NumHeaders = InNodalOffsetTable.Header.Num();
+
+	for (int32 DataRowIndex = 0; DataRowIndex < InNodalOffsetTable.Data.Num(); ++DataRowIndex)
+	{
+		const FLensFileParameterTableRow& DataRow = InNodalOffsetTable.Data[DataRowIndex];
+		const int32 NumDataValues = DataRow.Values.Num();
+
+		if (NumDataValues == NumHeaders)
+		{
+			float Focus = 0.0f;
+			float Zoom = 0.0f;
+			FQuat RotationOffset;
+			FVector LocationOffset;
+
+			for (int32 HeaderIndex = 0; HeaderIndex < NumHeaders; ++HeaderIndex)
+			{
+				const FName& HeaderName = InNodalOffsetTable.Header[HeaderIndex];
+				const float Value = DataRow.Values[HeaderIndex];
+
+				if (HeaderName == FocusEncoderHeaderName)
+				{
+					Focus = Value;
+				}
+				else if (HeaderName == ZoomEncoderHeaderName)
+				{
+					Zoom = Value;
+				}
+				else if (HeaderName == NodalOffsetQxHeaderName)
+				{
+					RotationOffset.X = Value;
+				}
+				else if (HeaderName == NodalOffsetQyHeaderName)
+				{
+					RotationOffset.Y = Value;
+				}
+				else if (HeaderName == NodalOffsetQzHeaderName)
+				{
+					RotationOffset.Z = Value;
+				}
+				else if (HeaderName == NodalOffsetQwHeaderName)
+				{
+					RotationOffset.W = Value;
+				}
+				else if (HeaderName == NodalOffsetTxHeaderName)
+				{
+					LocationOffset.X = Value;
+				}
+				else if (HeaderName == NodalOffsetTyHeaderName)
+				{
+					LocationOffset.Y = Value;
+				}
+				else if (HeaderName == NodalOffsetTzHeaderName)
+				{
+					LocationOffset.Z = Value;
+				}
+			}
+
+			FTransform NodalOffsetTransform{ RotationOffset, LocationOffset };
+
+			if (Metadata.NodalOffsetCoordinateSystem == ENodalOffsetCoordinateSystem::OpenCV)
+			{
+				FCameraCalibrationUtils::ConvertUnrealToOpenCV(NodalOffsetTransform);
+			}
+
+			FNodalPointOffset NodalPointOffset;
+			NodalPointOffset.LocationOffset = NodalOffsetTransform.GetLocation();
+			NodalPointOffset.RotationOffset = NodalOffsetTransform.GetRotation();
+			OutLensFile.AddNodalOffsetPoint(Focus, Zoom, NodalPointOffset);
+		}
+		else
+		{
+			GetFeedbackContext()->CategorizedLogf(LogCameraCalibrationEditor.GetCategoryName(),
+												  ELogVerbosity::Warning,
+												  TEXT("Number of values in row %d of table '%s' doesn't match with the number of headers. Expected %d values but got %d. Skipping line."),
+												  DataRowIndex, *InNodalOffsetTable.ParameterName.ToString(), NumHeaders, NumDataValues);
+		}
+	}
+}
+
+void FLensFileExchange::PopulateDistortionTable(ULensFile& OutLensFile, const FLensFileParameterTable& InDistortionTable) const
+{
+	const int32 NumHeaders = InDistortionTable.Header.Num();
+
+	for (int32 DataRowIndex = 0; DataRowIndex < InDistortionTable.Data.Num(); ++DataRowIndex)
+	{
+		const FLensFileParameterTableRow& DataRow = InDistortionTable.Data[DataRowIndex];
+		const int32 NumDataValues = DataRow.Values.Num();
+
+		if (NumDataValues == NumHeaders)
+		{
+			float Focus = 0.0f;
+			float Zoom = 0.0f;
+			TArray<float> DistortionParams;
+
+			for (int32 HeaderIndex = 0; HeaderIndex < NumHeaders; ++HeaderIndex)
+			{
+				const FName& HeaderName = InDistortionTable.Header[HeaderIndex];
+				const float Value = DataRow.Values[HeaderIndex];
+
+				if (HeaderName == FocusEncoderHeaderName)
+				{
+					Focus = Value;
+				}
+				else if (HeaderName == ZoomEncoderHeaderName)
+				{
+					Zoom = Value;
+				}
+				else
+				{
+					// TODO: Need to add some validation here based on ParameterDisplayNames, however, this function is WITH_EDITOR only
+					DistortionParams.Add(Value);
+				}
+			}
+
+			FDistortionInfo DistortionInfo;
+			DistortionInfo.Parameters = DistortionParams;
+			OutLensFile.DistortionTable.AddPoint(Focus, Zoom, DistortionInfo, OutLensFile.InputTolerance, false);
+		}
+		else
+		{
+			GetFeedbackContext()->CategorizedLogf(LogCameraCalibrationEditor.GetCategoryName(),
+												  ELogVerbosity::Warning,
+												  TEXT("Number of values in row %d of table '%s' doesn't match with the number of headers. Expected %d values but got %d. Skipping line."),
+												  DataRowIndex, *InDistortionTable.ParameterName.ToString(), NumHeaders, NumDataValues);
+		}
+	}
+}
+
+void FLensFileExchange::PopulateEncoderTable(FRichCurve& OutCurve, const FLensFileParameterTable& InEncoderTable, const TCHAR* InputHeaderName, const TCHAR* OutputHeaderName) const
+{
+	const int32 NumHeaders = InEncoderTable.Header.Num();
+
+	for (int32 DataRowIndex = 0; DataRowIndex < InEncoderTable.Data.Num(); ++DataRowIndex)
+	{
+		const FLensFileParameterTableRow& DataRow = InEncoderTable.Data[DataRowIndex];
+		const int32 NumDataValues = DataRow.Values.Num();
+
+		if (NumDataValues == NumHeaders)
+		{
+			float Input = 0.0f;
+			float Output = 0.0f;
+
+			for (int32 HeaderIndex = 0; HeaderIndex < NumHeaders; ++HeaderIndex)
+			{
+				const FName& HeaderName = InEncoderTable.Header[HeaderIndex];
+				const float Value = DataRow.Values[HeaderIndex];
+
+				if (HeaderName == InputHeaderName)
+				{
+					Input = Value;
+				}
+				else if (HeaderName == OutputHeaderName)
+				{
+					Output = Value;
+				}
+			}
+
+			OutCurve.AddKey(Input, Output);
+		}
+		else
+		{
+			GetFeedbackContext()->CategorizedLogf(LogCameraCalibrationEditor.GetCategoryName(),
+												  ELogVerbosity::Warning,
+												  TEXT("Number of values in row %d of table '%s' doesn't match with the number of headers. Expected %d values but got %d. Skipping line."),
+												  DataRowIndex, *InEncoderTable.ParameterName.ToString(), NumHeaders, NumDataValues);
+		}
+	}
+}
+
+FFeedbackContext* FLensFileExchange::GetFeedbackContext() const
+{
+	return FeedbackContext ? FeedbackContext : GWarn;
+}
+
+#undef LOCTEXT_NAMESPACE
\ No newline at end of file
diff --git a/Engine/Plugins/VirtualProduction/CameraCalibration/Source/CameraCalibrationEditor/Private/LensFileExchangeFormat.h b/Engine/Plugins/VirtualProduction/CameraCalibration/Source/CameraCalibrationEditor/Private/LensFileExchangeFormat.h
new file mode 100644
--- /dev/null
+++ b/Engine/Plugins/VirtualProduction/CameraCalibration/Source/CameraCalibrationEditor/Private/LensFileExchangeFormat.h
@@ -0,0 +1,223 @@
+// Copyright Epic Games, Inc. All Rights Reserved.
+
+#pragma once
+
+#include "CoreMinimal.h"
+#include "UObject/ObjectMacros.h"
+
+#include "LensFileExchangeFormat.generated.h"
+
+UENUM()
+enum class ELensFileUnit
+{
+	Millimeters,
+	Normalized
+};
+
+UENUM()
+enum class ENodalOffsetCoordinateSystem
+{
+	OpenCV
+};
+
+USTRUCT()
+struct CAMERACALIBRATIONEDITOR_API FLensInfoExchange
+{
+	GENERATED_BODY()
+
+	FLensInfoExchange(const class ULensFile* LensFile = nullptr);
+
+	UPROPERTY()
+	FName SerialNumber;
+
+	UPROPERTY()
+	FName ModelName;
+
+	UPROPERTY()
+	FName DistortionModel;
+};
+
+USTRUCT()
+struct FLensFileUserMetadataEntry
+{
+	GENERATED_BODY()
+
+	UPROPERTY()
+	FName Name;
+
+	UPROPERTY()
+	FName Value;
+};
+
+USTRUCT()
+struct CAMERACALIBRATIONEDITOR_API FLensFileMetadata
+{
+	GENERATED_BODY()
+
+	FLensFileMetadata(const class ULensFile* LensFile = nullptr);
+
+	UPROPERTY()
+	FName Type;
+
+	UPROPERTY()
+	FName Version;
+
+	UPROPERTY()
+	FLensInfoExchange LensInfo;
+
+	UPROPERTY()
+	FName Name;
+
+	UPROPERTY()
+	ENodalOffsetCoordinateSystem NodalOffsetCoordinateSystem = ENodalOffsetCoordinateSystem::OpenCV;
+
+	UPROPERTY()
+	ELensFileUnit FxFyUnits = ELensFileUnit::Normalized;
+
+	UPROPERTY()
+	ELensFileUnit CxCyUnits = ELensFileUnit::Normalized;
+
+	UPROPERTY()
+	TArray<FLensFileUserMetadataEntry> UserMetadata;
+};
+
+
+USTRUCT()
+struct FLensFileSensorDimensions
+{
+	GENERATED_BODY()
+
+	UPROPERTY()
+	float Width = 0.0f;
+
+	UPROPERTY()
+	float Height = 0.0f;
+
+	UPROPERTY()
+	ELensFileUnit Units = ELensFileUnit::Normalized;
+};
+
+USTRUCT()
+struct FLensFileImageDimensions
+{
+	GENERATED_BODY()
+
+	UPROPERTY()
+	int32 Width = 0;
+
+	UPROPERTY()
+	int32 Height = 0;
+};
+
+USTRUCT()
+struct FLensFileParameterTableRow
+{
+	GENERATED_BODY()
+
+	explicit FLensFileParameterTableRow(int32 InNumHeaders = 0);
+
+	UPROPERTY()
+	TArray<float> Values;
+};
+
+USTRUCT()
+struct FLensFileParameterTable
+{
+	GENERATED_BODY()
+
+	UPROPERTY()
+	FName ParameterName;
+
+	UPROPERTY()
+	TArray<FName> Header;
+
+	UPROPERTY()
+	TArray<FLensFileParameterTableRow> Data;
+};
+
+USTRUCT()
+struct FLensFileParameterTableImporter
+{
+	GENERATED_BODY()
+
+	UPROPERTY()
+	FName ParameterName;
+
+	UPROPERTY()
+	FString Header;
+
+	UPROPERTY()
+	FString Data;
+
+	bool ConvertToParameterTable(struct FLensFileParameterTable& OutParamterTable, class FFeedbackContext* InWarn) const;
+};
+
+USTRUCT()
+struct CAMERACALIBRATIONEDITOR_API FLensFileExchange
+{
+	GENERATED_BODY()
+
+	FLensFileExchange(const class ULensFile* LensFile = nullptr);
+
+	FLensFileExchange(const TSharedRef<FJsonObject>& InLensFileJson, bool& bOutIsValidJson, class FFeedbackContext* InWarn);
+
+	bool PopulateLensFile(class ULensFile& LensFile) const;
+
+	UPROPERTY()
+	FLensFileMetadata Metadata;
+
+	UPROPERTY()
+	FLensFileSensorDimensions SensorDimensions;
+
+	UPROPERTY()
+	FLensFileImageDimensions ImageDimensions;
+
+	UPROPERTY()
+	TArray<FLensFileParameterTable> CameraParameterTables;
+
+	UPROPERTY()
+	TArray<FLensFileParameterTable> EncoderTables;
+
+	constexpr static TCHAR LensFileType[] = TEXT("LensFile");
+	constexpr static TCHAR LensFileVersion[] = TEXT("0.0.0");
+
+	constexpr static TCHAR FocusEncoderHeaderName[] = TEXT("FocusEncoder");
+	constexpr static TCHAR ZoomEncoderHeaderName[] = TEXT("ZoomEncoder");
+	constexpr static TCHAR IrisEncoderHeaderName[] = TEXT("IrisEncoder");
+
+	constexpr static TCHAR FocusCMHeaderName[] = TEXT("FocusCM");
+	constexpr static TCHAR IrisFstopHeaderName[] = TEXT("IrisFstop");
+
+	constexpr static TCHAR FocalLengthFxHeaderName[] = TEXT("Fx");
+	constexpr static TCHAR FocalLengthFyHeaderName[] = TEXT("Fy");
+
+	constexpr static TCHAR ImageCenterCxHeaderName[] = TEXT("Cx");
+	constexpr static TCHAR ImageCenterCyHeaderName[] = TEXT("Cy");
+
+	constexpr static TCHAR NodalOffsetQxHeaderName[] = TEXT("Qx");
+	constexpr static TCHAR NodalOffsetQyHeaderName[] = TEXT("Qy");
+	constexpr static TCHAR NodalOffsetQzHeaderName[] = TEXT("Qz");
+	constexpr static TCHAR NodalOffsetQwHeaderName[] = TEXT("Qw");
+
+	constexpr static TCHAR NodalOffsetTxHeaderName[] = TEXT("Tx");
+	constexpr static TCHAR NodalOffsetTyHeaderName[] = TEXT("Ty");
+	constexpr static TCHAR NodalOffsetTzHeaderName[] = TEXT("Tz");
+
+private:
+	FFeedbackContext* FeedbackContext;
+
+	FFeedbackContext* GetFeedbackContext() const;
+
+	void ExtractFocalLengthTable(const class ULensFile* LensFile);
+	void ExtractImageCenterTable(const class ULensFile* LensFile);
+	void ExtractNodalOffsetTable(const class ULensFile* LensFile);
+	void ExtractEncoderTables(const class ULensFile* LensFile);
+	void ExtractDistortionParameters(const class ULensFile* LensFile);
+	void ExtractSTMaps(const class ULensFile* LensFile);
+
+	void PopulateFocalLengthTable(class ULensFile& OutLensFile, const FLensFileParameterTable& FocalLengthTable) const;
+	void PopulateImageCenterTable(class ULensFile& OutLensFile, const FLensFileParameterTable& InImageCenterTable) const;
+	void PopulateNodalOffsetTable(class ULensFile& OutLensFile, const FLensFileParameterTable& InNodalOffsetTable) const;
+	void PopulateDistortionTable(class ULensFile& OutLensFile, const FLensFileParameterTable& InDistortionTable) const;
+	void PopulateEncoderTable(struct FRichCurve& OutCurve, const FLensFileParameterTable& InEncoderTable, const TCHAR* InputHeaderName, const TCHAR* OutputHeaderName) const;
+};
diff --git a/Engine/Plugins/VirtualProduction/CameraCalibrationCore/Source/CameraCalibrationCore/Private/LensFile.cpp b/Engine/Plugins/VirtualProduction/CameraCalibrationCore/Source/CameraCalibrationCore/Private/LensFile.cpp
--- a/Engine/Plugins/VirtualProduction/CameraCalibrationCore/Source/CameraCalibrationCore/Private/LensFile.cpp
+++ b/Engine/Plugins/VirtualProduction/CameraCalibrationCore/Source/CameraCalibrationCore/Private/LensFile.cpp
@@ -8,6 +8,7 @@
 #include "CameraCalibrationSettings.h"
 #include "CameraCalibrationSubsystem.h"
 #include "CineCameraComponent.h"
+#include "EditorFramework/AssetImportData.h"
 #include "Engine/Engine.h"
 #include "Engine/TextureRenderTarget2D.h"
 #include "LensDistortionModelHandlerBase.h"
@@ -1306,6 +1307,13 @@ void ULensFile::PostInitProperties()
 	DistortionTable.LensFile =
 		FocalLengthTable.LensFile = ImageCenterTable.LensFile =
 		NodalOffsetTable.LensFile = STMapTable.LensFile = this;
+
+#if WITH_EDITORONLY_DATA
+	if (!HasAnyFlags(RF_ClassDefaultObject))
+	{
+		AssetImportData = NewObject<UAssetImportData>(this, TEXT("AssetImportData"));
+	}
+#endif
 }
 
 void ULensFile::Tick(float DeltaTime)
diff --git a/Engine/Plugins/VirtualProduction/CameraCalibrationCore/Source/CameraCalibrationCore/Public/LensFile.h b/Engine/Plugins/VirtualProduction/CameraCalibrationCore/Source/CameraCalibrationCore/Public/LensFile.h
--- a/Engine/Plugins/VirtualProduction/CameraCalibrationCore/Source/CameraCalibrationCore/Public/LensFile.h
+++ b/Engine/Plugins/VirtualProduction/CameraCalibrationCore/Source/CameraCalibrationCore/Public/LensFile.h
@@ -267,6 +267,12 @@ public:
 	/** Tolerance used to consider input focus or zoom to be identical */
 	float InputTolerance = KINDA_SMALL_NUMBER;
 
+#if WITH_EDITORONLY_DATA
+	/** Importing data and options used for importing ulens files. */
+	UPROPERTY(VisibleAnywhere, Instanced, Category = "ImportSettings")
+	TObjectPtr<UAssetImportData> AssetImportData;
+#endif // WITH_EDITORONLY_DATA
+
 protected:
 
 	/** Derived data compute jobs we are waiting on */
