Fix for non-unity build

diff --git a/Engine/Plugins/VirtualProduction/CameraCalibration/Source/CameraCalibrationEditor/Private/Factories/LensFileImporterFactory.cpp b/Engine/Plugins/VirtualProduction/CameraCalibration/Source/CameraCalibrationEditor/Private/Factories/LensFileImporterFactory.cpp
--- a/Engine/Plugins/VirtualProduction/CameraCalibration/Source/CameraCalibrationEditor/Private/Factories/LensFileImporterFactory.cpp
+++ b/Engine/Plugins/VirtualProduction/CameraCalibration/Source/CameraCalibrationEditor/Private/Factories/LensFileImporterFactory.cpp
@@ -3,6 +3,7 @@
 #include "LensFileImporterFactory.h"
 
 #include "CoreMinimal.h"
+#include "Editor.h"
 #include "EditorFramework/AssetImportData.h"
 #include "JsonObjectConverter.h"
 #include "LensFile.h"
diff --git a/Engine/Plugins/VirtualProduction/CameraCalibration/Source/CameraCalibrationEditor/Private/LensFileExchangeFormat.h b/Engine/Plugins/VirtualProduction/CameraCalibration/Source/CameraCalibrationEditor/Private/LensFileExchangeFormat.h
--- a/Engine/Plugins/VirtualProduction/CameraCalibration/Source/CameraCalibrationEditor/Private/LensFileExchangeFormat.h
+++ b/Engine/Plugins/VirtualProduction/CameraCalibration/Source/CameraCalibrationEditor/Private/LensFileExchangeFormat.h
@@ -159,7 +159,7 @@ struct CAMERACALIBRATIONEDITOR_API FLensFileExchange
 
 	FLensFileExchange(const class ULensFile* LensFile = nullptr);
 
-	FLensFileExchange(const TSharedRef<FJsonObject>& InLensFileJson, bool& bOutIsValidJson, class FFeedbackContext* InWarn);
+	FLensFileExchange(const TSharedRef<class FJsonObject>& InLensFileJson, bool& bOutIsValidJson, class FFeedbackContext* InWarn);
 
 	bool PopulateLensFile(class ULensFile& LensFile) const;
 
