[CameraCalibration] Improve accuracy of nodal offset calibration using an opencv downhill solver to minimize reprojection error for multiple camera views

diff --git a/Engine/Plugins/Runtime/OpenCV/Source/OpenCVHelper/Private/OpenCVHelper.cpp b/Engine/Plugins/Runtime/OpenCV/Source/OpenCVHelper/Private/OpenCVHelper.cpp
--- a/Engine/Plugins/Runtime/OpenCV/Source/OpenCVHelper/Private/OpenCVHelper.cpp
+++ b/Engine/Plugins/Runtime/OpenCV/Source/OpenCVHelper/Private/OpenCVHelper.cpp
@@ -189,5 +189,53 @@ UTexture2D* FOpenCVHelper::TextureFromCvMat(cv::Mat& Mat, UTexture2D* InTexture)
 	return InTexture;
 }
 
+double FOpenCVHelper::ComputeReprojectionError(const FTransform& CameraPose, const cv::Mat& CameraIntrinsicMatrix, const std::vector<cv::Point3f>& Points3d, const std::vector<cv::Point2f>& Points2d)
+{
+	// Ensure that the number of point correspondences is valid
+	const int32 NumPoints3d = Points3d.size();
+	const int32 NumPoints2d = Points2d.size();
+	if ((NumPoints3d == 0) || (NumPoints2d == 0) || (NumPoints3d != NumPoints2d))
+	{
+		return -1.0;
+	}
+
+	const FMatrix CameraPoseMatrix = CameraPose.ToMatrixNoScale();
+
+	const cv::Mat Tcam = (cv::Mat_<double>(3, 1) << CameraPoseMatrix.M[3][0], CameraPoseMatrix.M[3][1], CameraPoseMatrix.M[3][2]);
+
+	cv::Mat Rcam = cv::Mat::zeros(3, 3, cv::DataType<double>::type);
+	for (int32 Column = 0; Column < 3; ++Column)
+	{
+		FVector ColVec = CameraPoseMatrix.GetColumn(Column);
+		Rcam.at<double>(Column, 0) = ColVec.X;
+		Rcam.at<double>(Column, 1) = ColVec.Y;
+		Rcam.at<double>(Column, 2) = ColVec.Z;
+	}
+
+	const cv::Mat Robj = Rcam.t();
+
+	cv::Mat Rrod;
+	cv::Rodrigues(Robj, Rrod);
+
+	const cv::Mat Tobj = -Rcam.inv() * Tcam;
+
+	std::vector<cv::Point2f> ReprojectedPoints2d;
+
+	// The 2D points will be compared against the undistorted 2D points, so the distortion coefficients can be ignored
+	cv::projectPoints(Points3d, Rrod, Tobj, CameraIntrinsicMatrix, cv::noArray(), ReprojectedPoints2d);
+
+	// Compute euclidean distance between captured 2D points and reprojected 2D points to measure reprojection error
+	double ReprojectionError = 0.0;
+	for (int32 Index = 0; Index < NumPoints2d; ++Index)
+	{
+		const cv::Point2f& A = Points2d[Index];
+		const cv::Point2f& B = ReprojectedPoints2d[Index];
+		const cv::Point2f Diff = A - B;
+
+		ReprojectionError += (Diff.x * Diff.x) + (Diff.y * Diff.y); // cv::norm with NORM_L2SQR
+	}
+
+	return ReprojectionError;
+}
 
 #endif // WITH_OPENCV
diff --git a/Engine/Plugins/Runtime/OpenCV/Source/OpenCVHelper/Public/OpenCVHelper.h b/Engine/Plugins/Runtime/OpenCV/Source/OpenCVHelper/Public/OpenCVHelper.h
--- a/Engine/Plugins/Runtime/OpenCV/Source/OpenCVHelper/Public/OpenCVHelper.h
+++ b/Engine/Plugins/Runtime/OpenCV/Source/OpenCVHelper/Public/OpenCVHelper.h
@@ -2,7 +2,7 @@
 
 #pragma once
 #include "CoreMinimal.h"
-
+#include <vector>
 /*
 Like many third party headers, OpenCV headers require some care when importing.
 
@@ -47,6 +47,11 @@ class FName;
 namespace cv
 {
 	class Mat;
+	template<typename _Tp> class Point_;
+	template<typename _Tp> class Point3_;
+
+	typedef Point_<float> Point2f;
+	typedef Point3_<float> Point3f;
 };
 
 class OPENCVHELPER_API FOpenCVHelper
@@ -64,6 +69,8 @@ public:
 	 */
 	static UTexture2D* TextureFromCvMat(cv::Mat& Mat, const FString* PackagePath = nullptr, const FName* TextureName = nullptr);
 	static UTexture2D* TextureFromCvMat(cv::Mat& Mat, UTexture2D* InTexture);
+
+	static double ComputeReprojectionError(const FTransform& CameraPose, const cv::Mat& CameraIntrinsicMatrix, const std::vector<cv::Point3f>& Points3d, const std::vector<cv::Point2f>& Points2d);
 };
 
 #endif //WITH_OPENCV
diff --git a/Engine/Plugins/VirtualProduction/CameraCalibration/Source/CameraCalibrationEditor/Private/AssetEditor/LensDistortionTool.cpp b/Engine/Plugins/VirtualProduction/CameraCalibration/Source/CameraCalibrationEditor/Private/AssetEditor/LensDistortionTool.cpp
--- a/Engine/Plugins/VirtualProduction/CameraCalibration/Source/CameraCalibrationEditor/Private/AssetEditor/LensDistortionTool.cpp
+++ b/Engine/Plugins/VirtualProduction/CameraCalibration/Source/CameraCalibrationEditor/Private/AssetEditor/LensDistortionTool.cpp
@@ -209,7 +209,7 @@ void ULensDistortionTool::OnSaveCurrentCalibrationData()
 	// Show reprojection error
 	{
 		FFormatOrderedArguments Arguments;
-		Arguments.Add(FText::FromString(FString::Printf(TEXT("%.1f"), Error)));
+		Arguments.Add(FText::FromString(FString::Printf(TEXT("%.2f"), Error)));
 
 		const FText Message = FText::Format(LOCTEXT("ReprojectionError", "RMS Reprojection Error: {0} pixels"), Arguments);
 
diff --git a/Engine/Plugins/VirtualProduction/CameraCalibration/Source/CameraCalibrationEditor/Private/AssetEditor/NodalOffsetTool.cpp b/Engine/Plugins/VirtualProduction/CameraCalibration/Source/CameraCalibrationEditor/Private/AssetEditor/NodalOffsetTool.cpp
--- a/Engine/Plugins/VirtualProduction/CameraCalibration/Source/CameraCalibrationEditor/Private/AssetEditor/NodalOffsetTool.cpp
+++ b/Engine/Plugins/VirtualProduction/CameraCalibration/Source/CameraCalibrationEditor/Private/AssetEditor/NodalOffsetTool.cpp
@@ -160,6 +160,7 @@ void UNodalOffsetTool::OnSaveCurrentNodalOffset()
 		return;
 	}
 
+	const FText TitleInfo = LOCTEXT("NodalOffsetInfo", "Nodal Offset Calibration Info");
 	const FText TitleError = LOCTEXT("NodalOffsetError", "Nodal Offset Calibration Error");
 
 	UCameraNodalOffsetAlgo* Algo = GetNodalOffsetAlgo();
@@ -184,6 +185,20 @@ void UNodalOffsetTool::OnSaveCurrentNodalOffset()
 		return;
 	}
 
+	// Show reprojection error
+	{
+		FFormatOrderedArguments Arguments;
+		Arguments.Add(FText::FromString(FString::Printf(TEXT("%.2f"), ReprojectionError)));
+
+		const FText Message = FText::Format(LOCTEXT("ReprojectionError", "RMS Reprojection Error: {0} pixels"), Arguments);
+
+		// Allow the user to cancel adding to the LUT if the reprojection error is unacceptable.
+		if (FMessageDialog::Open(EAppMsgType::OkCancel, Message, &TitleInfo) != EAppReturnType::Ok)
+		{
+			return;
+		}
+	}
+
 	ULensFile* LensFile = CameraCalibrationStepsController.Pin()->GetLensFile();
 
 	if (!LensFile)
diff --git a/Engine/Plugins/VirtualProduction/CameraCalibration/Source/CameraCalibrationEditor/Private/Calibrators/CameraNodalOffsetAlgoPoints.cpp b/Engine/Plugins/VirtualProduction/CameraCalibration/Source/CameraCalibrationEditor/Private/Calibrators/CameraNodalOffsetAlgoPoints.cpp
--- a/Engine/Plugins/VirtualProduction/CameraCalibration/Source/CameraCalibrationEditor/Private/Calibrators/CameraNodalOffsetAlgoPoints.cpp
+++ b/Engine/Plugins/VirtualProduction/CameraCalibration/Source/CameraCalibrationEditor/Private/Calibrators/CameraNodalOffsetAlgoPoints.cpp
@@ -19,6 +19,7 @@
 #include "LensDistortionModelHandlerBase.h"
 #include "LensFile.h"
 #include "Math/Vector.h"
+#include "Misc/FileHelper.h"
 #include "Misc/MessageDialog.h"
 #include "PropertyCustomizationHelpers.h"
 #include "UI/SFilterableActorPicker.h"
@@ -41,6 +42,11 @@ OPENCV_INCLUDES_END
 
 #define LOCTEXT_NAMESPACE "CameraNodalOffsetAlgoPoints"
 
+#if WITH_EDITOR
+static TAutoConsoleVariable<float> CVarRotationStepValue(TEXT("CameraCalibration.RotationStepValue"), 0.05, TEXT("The value of the initial step size to use when finding an optimal nodal offset rotation that minimizes reprojection error."));
+static TAutoConsoleVariable<float> CVarLocationStepValue(TEXT("CameraCalibration.LocationStepValue"), 0.5, TEXT("The value of the initial step size to use when finding an optimal nodal offset location that minimizes reprojection error."));
+#endif
+
 namespace CameraNodalOffsetAlgoPoints
 {
 	class SCalibrationRowGenerator
@@ -189,6 +195,334 @@ namespace CameraNodalOffsetAlgoPoints
 	}
 };
 
+#if WITH_OPENCV
+class FMinReprojectionErrorSolver : public cv::MinProblemSolver::Function
+{
+public:
+	struct FSingleViewPoseAndPoints
+	{
+		FTransform CameraPose;
+
+		std::vector<cv::Point3f> Points3d;
+		std::vector<cv::Point2f> Points2d;
+	};
+
+	//~ Begin cv::MinProblemSolver::Function interface
+	int getDims() const
+	{
+		return 7; // 4 doubles for the rotation quaternion + 3 doubles for the location vector
+	}
+
+	double calc(const double* x) const
+	{
+		// Convert the input data (7 doubles) to an FQuat and FVector
+		const FQuat Rotation = FQuat(x[0], x[1], x[2], x[3]).GetNormalized();
+		const FVector Location = FVector(x[4], x[5], x[6]);
+
+		UE_LOG(LogCameraCalibrationEditor, VeryVerbose, TEXT("Nodal Offset Candidate:  Rotation: (%lf, %lf, %lf, %lf)  Location: (%lf, %lf, %lf)"),
+			Rotation.X, Rotation.Y, Rotation.Z, Rotation.W, Location.X, Location.Y, Location.Z);
+
+		FTransform NodalOffsetCandidate;
+
+		// As a result of the way that the downhill solver nudges the input data on each iteration, it is important to normalize the rotation
+		NodalOffsetCandidate.SetRotation(Rotation);
+		NodalOffsetCandidate.SetLocation(Location);
+
+		double ReprojectionErrorTotal = 0.0;
+		int32 NumTotalPoints = 0;
+
+		// For each camera view
+		for (const FSingleViewPoseAndPoints& CameraView : CameraViews)
+		{
+			// Compute the optimal camera pose using the nodal offset candidate for this iteration and the tracked camera pose 
+			// Any existing nodal offset that was applied to the camera is also factored in
+			FTransform OptimalCameraPose = NodalOffsetCandidate * ExistingNodalOffsetInverse * CameraView.CameraPose;
+
+			// Convert the optimal camera pose to opencv's coordinate system
+			FCameraCalibrationUtils::ConvertUnrealToOpenCV(OptimalCameraPose);
+
+			// Compute the reprojection error for this view and add it to the running total
+			double ViewError = FOpenCVHelper::ComputeReprojectionError(OptimalCameraPose, CameraMatrix, CameraView.Points3d, CameraView.Points2d);
+			ReprojectionErrorTotal += ViewError;
+
+			NumTotalPoints += CameraView.Points2d.size();
+		}
+
+		double RootMeanSquareError = FMath::Sqrt(ReprojectionErrorTotal / NumTotalPoints);
+
+		UE_LOG(LogCameraCalibrationEditor, VeryVerbose, TEXT("Reprojection Error: %lf"), RootMeanSquareError);
+
+		return RootMeanSquareError;
+	}
+	//~ End cv::MinProblemSolver::Function interface
+
+	void SetCameraMatrix(const cv::Mat InCameraMatrix) 
+	{ 
+		CameraMatrix = InCameraMatrix; 
+	}
+
+	void SetExistingNodalOffsetInverse(FTransform InOffsetInverse) 
+	{ 
+		ExistingNodalOffsetInverse = InOffsetInverse; 
+	}
+
+	void AddCameraViewAndPoints(FTransform InCameraPose, std::vector<cv::Point3f> InPoints3d, std::vector<cv::Point2f> InPoints2d)
+	{
+		FSingleViewPoseAndPoints& NewCameraView = CameraViews.Add_GetRef(FSingleViewPoseAndPoints());
+
+		NewCameraView.CameraPose = InCameraPose;
+		NewCameraView.Points3d = MoveTemp(InPoints3d);
+		NewCameraView.Points2d = MoveTemp(InPoints2d);
+	}
+
+private:
+	cv::Mat CameraMatrix;
+
+	FTransform ExistingNodalOffsetInverse;
+
+	TArray<FSingleViewPoseAndPoints> CameraViews;
+};
+#endif
+
+double UCameraNodalOffsetAlgoPoints::MinimizeReprojectionError(FTransform& InOutNodalOffset, const TArray<TSharedPtr<TArray<TSharedPtr<FCalibrationRowData>>>>& SamePoseRowGroups) const
+{
+#if WITH_OPENCV
+	cv::Ptr<cv::DownhillSolver> Solver = cv::DownhillSolver::create();
+
+	cv::Ptr<FMinReprojectionErrorSolver> SolverFunction = cv::makePtr<FMinReprojectionErrorSolver>();
+	Solver->setFunction(SolverFunction);
+
+	const FQuat InitialRotation = InOutNodalOffset.GetRotation();
+	const FVector InitialLocation = InOutNodalOffset.GetLocation();
+
+	cv::Mat WorkingSolution(1, 7, cv::DataType<double>::type);
+	WorkingSolution.at<double>(0, 0) = InitialRotation.X;
+	WorkingSolution.at<double>(0, 1) = InitialRotation.Y;
+	WorkingSolution.at<double>(0, 2) = InitialRotation.Z;
+	WorkingSolution.at<double>(0, 3) = InitialRotation.W;
+	WorkingSolution.at<double>(0, 4) = InitialLocation.X;
+	WorkingSolution.at<double>(0, 5) = InitialLocation.Y;
+	WorkingSolution.at<double>(0, 6) = InitialLocation.Z;
+
+	// NOTE: These step sizes may need further testing and refinement, but tests so far have shown them to be decent
+	const double RotationStep = CVarRotationStepValue.GetValueOnAnyThread();
+	const double LocationStep = CVarLocationStepValue.GetValueOnAnyThread();
+	cv::Mat Step = (cv::Mat_<double>(7, 1) << RotationStep, RotationStep, RotationStep, RotationStep, LocationStep, LocationStep, LocationStep);
+
+	Solver->setInitStep(Step);
+
+	const FCameraCalibrationStepsController* StepsController;
+	const ULensFile* LensFile;
+
+	if (!ensure(GetStepsControllerAndLensFile(&StepsController, &LensFile)))
+	{
+		return -1.0;
+	}
+
+	const ULensDistortionModelHandlerBase* DistortionHandler = StepsController->GetDistortionHandler();
+	if (!DistortionHandler)
+	{
+		return -1.0;
+	}
+
+	const FLensDistortionState DistortionState = DistortionHandler->GetCurrentDistortionState();
+
+	// Initialize the camera matrix that will be used in each call to projectPoints()
+	cv::Mat CameraMatrix(3, 3, cv::DataType<double>::type);
+	cv::setIdentity(CameraMatrix);
+
+	CameraMatrix.at<double>(0, 0) = DistortionState.FocalLengthInfo.FxFy.X;
+	CameraMatrix.at<double>(1, 1) = DistortionState.FocalLengthInfo.FxFy.Y;
+
+	// Note that the 2D points that will be compared against have been undistorted and center shift has already been accounted for
+	CameraMatrix.at<double>(0, 2) = 0.5;
+	CameraMatrix.at<double>(1, 2) = 0.5;
+
+	SolverFunction->SetCameraMatrix(CameraMatrix);
+
+	if (!ensureMsgf((CalibrationRows.Num() > 0), TEXT("Not enough calibration rows")))
+	{
+		return -1.0;
+	}
+
+	const TSharedPtr<FCalibrationRowData>& FirstRow = CalibrationRows[0];
+
+	FTransform ExistingOffset = FTransform::Identity;
+
+	if (FirstRow->CameraData.LensFileEvalData.NodalOffset.bWasApplied)
+	{
+		FNodalPointOffset NodalPointOffset;
+		const float Focus = FirstRow->CameraData.LensFileEvalData.Input.Focus;
+		const float Zoom = FirstRow->CameraData.LensFileEvalData.Input.Zoom;
+
+		if (LensFile->EvaluateNodalPointOffset(Focus, Zoom, NodalPointOffset))
+		{
+			ExistingOffset.SetTranslation(NodalPointOffset.LocationOffset);
+			ExistingOffset.SetRotation(NodalPointOffset.RotationOffset);
+		}
+	}
+
+	SolverFunction->SetExistingNodalOffsetInverse(ExistingOffset.Inverse());
+
+	if (!ensureMsgf(FirstRow->bUndistortedIsValid, TEXT("This method operates on undistorted 2D points. Call UndistortCalibrationRowPoints() prior to calling this method")))
+	{
+		return -1.0;
+	}
+
+	for (int32 PoseGroupIndex = 0; PoseGroupIndex < SamePoseRowGroups.Num(); ++PoseGroupIndex)
+	{
+		const TArray<TSharedPtr<FCalibrationRowData>>& PoseGroup = *(SamePoseRowGroups[PoseGroupIndex]);
+
+		if (!ensureMsgf((PoseGroup.Num() > 0), TEXT("Not enough calibration rows")))
+		{
+			return -1.0;
+		}
+
+		const TSharedPtr<FCalibrationRowData>& FirstRowInPoseGroup = PoseGroup[0];
+
+		FTransform CameraPose = FirstRowInPoseGroup->CameraData.Pose;
+
+		std::vector<cv::Point3f> Points3d;
+		std::vector<cv::Point2f> Points2d;
+		Points3d.reserve(PoseGroup.Num());
+		Points2d.reserve(PoseGroup.Num());
+
+		for (const TSharedPtr<FCalibrationRowData>& Row : PoseGroup)
+		{
+			// Convert from UE coordinates to OpenCV coordinates
+			FTransform Transform;
+			Transform.SetIdentity();
+			Transform.SetLocation(Row->CalibratorPointData.Location);
+
+			FCameraCalibrationUtils::ConvertUnrealToOpenCV(Transform);
+
+			// Calibrator 3d points
+			Points3d.push_back(cv::Point3f(
+				Transform.GetLocation().X,
+				Transform.GetLocation().Y,
+				Transform.GetLocation().Z));
+
+			Points2d.push_back(cv::Point2f(
+				Row->UndistortedPoint2D.X,
+				Row->UndistortedPoint2D.Y));
+		}
+
+		SolverFunction->AddCameraViewAndPoints(CameraPose, Points3d, Points2d);
+	}
+
+	double Error = Solver->minimize(WorkingSolution);
+
+	const FQuat FinalRotation = FQuat(WorkingSolution.at<double>(0, 0), WorkingSolution.at<double>(0, 1), WorkingSolution.at<double>(0, 2), WorkingSolution.at<double>(0, 3)).GetNormalized();
+	const FVector FinalLocation = FVector(WorkingSolution.at<double>(0, 4), WorkingSolution.at<double>(0, 5), WorkingSolution.at<double>(0, 6));
+
+	InOutNodalOffset.SetRotation(FinalRotation);
+	InOutNodalOffset.SetLocation(FinalLocation);
+
+	return Error;
+#else
+	return -1.0;
+#endif
+}
+
+double UCameraNodalOffsetAlgoPoints::ComputeReprojectionError(const FTransform& NodalOffset, const TArray<TSharedPtr<FCalibrationRowData>>& PoseGroup) const
+{
+#if WITH_OPENCV
+	const FCameraCalibrationStepsController* StepsController;
+	const ULensFile* LensFile;
+
+	if (!ensure(GetStepsControllerAndLensFile(&StepsController, &LensFile)))
+	{
+		return -1.0;
+	}
+
+	const ULensDistortionModelHandlerBase* DistortionHandler = StepsController->GetDistortionHandler();
+	if (!DistortionHandler)
+	{
+		return -1.0;
+	}
+
+	std::vector<cv::Point3f> Points3d;
+	std::vector<cv::Point2f> Points2d;
+	Points3d.reserve(PoseGroup.Num());
+	Points2d.reserve(PoseGroup.Num());
+
+	if (!ensureMsgf((PoseGroup.Num() > 0), TEXT("Not enough calibration rows")))
+	{
+		return -1.0;
+	}
+
+	const TSharedPtr<FCalibrationRowData>& FirstRow = PoseGroup[0];
+
+	if (!ensureMsgf(FirstRow->bUndistortedIsValid, TEXT("This method operates on undistorted 2D points. Call UndistortCalibrationRowPoints() prior to calling this method")))
+	{
+		return -1.0;
+	}
+
+	for (const TSharedPtr<FCalibrationRowData>& Row : PoseGroup)
+	{
+		// Convert from UE coordinates to OpenCV coordinates
+
+		FTransform Transform;
+		Transform.SetIdentity();
+		Transform.SetLocation(Row->CalibratorPointData.Location);
+
+		FCameraCalibrationUtils::ConvertUnrealToOpenCV(Transform);
+
+		// Calibrator 3d points
+		Points3d.push_back(cv::Point3f(
+			Transform.GetLocation().X,
+			Transform.GetLocation().Y,
+			Transform.GetLocation().Z));
+
+		Points2d.push_back(cv::Point2f(
+			Row->UndistortedPoint2D.X,
+			Row->UndistortedPoint2D.Y));
+	}
+
+	// Compute the optimal camera pose using the nodal offset candidate for this iteration and the tracked camera pose 
+	// Any existing nodal offset that was applied to the camera is also factored in
+	FTransform CameraPose = FirstRow->CameraData.Pose;
+
+	FTransform ExistingOffset = FTransform::Identity;
+
+	if (FirstRow->CameraData.LensFileEvalData.NodalOffset.bWasApplied)
+	{
+		FNodalPointOffset NodalPointOffset;
+		const float Focus = FirstRow->CameraData.LensFileEvalData.Input.Focus;
+		const float Zoom = FirstRow->CameraData.LensFileEvalData.Input.Zoom;
+
+		if (LensFile->EvaluateNodalPointOffset(Focus, Zoom, NodalPointOffset))
+		{
+			ExistingOffset.SetTranslation(NodalPointOffset.LocationOffset);
+			ExistingOffset.SetRotation(NodalPointOffset.RotationOffset);
+		}
+	}
+
+	// Initialize the camera matrix that will be used in each call to projectPoints()
+	FLensDistortionState DistortionState = DistortionHandler->GetCurrentDistortionState();
+
+	cv::Mat CameraMatrix(3, 3, cv::DataType<double>::type);
+	cv::setIdentity(CameraMatrix);
+
+	CameraMatrix.at<double>(0, 0) = DistortionState.FocalLengthInfo.FxFy.X;
+	CameraMatrix.at<double>(1, 1) = DistortionState.FocalLengthInfo.FxFy.Y;
+
+	// Note that the 2D points that will be compared against have been undistorted and center shift has already been accounted for
+	CameraMatrix.at<double>(0, 2) = 0.5;
+	CameraMatrix.at<double>(1, 2) = 0.5;
+
+	FTransform OptimalCameraPose = NodalOffset * ExistingOffset.Inverse() * CameraPose;
+
+	// Convert the optimal camera pose to opencv's coordinate system
+	FCameraCalibrationUtils::ConvertUnrealToOpenCV(OptimalCameraPose);
+
+	// Compute the reprojection error for this view
+	return FOpenCVHelper::ComputeReprojectionError(OptimalCameraPose, CameraMatrix, Points3d, Points2d);
+#else
+	return -1.0;
+#endif
+}
+
 void UCameraNodalOffsetAlgoPoints::Initialize(UNodalOffsetTool* InNodalOffsetTool)
 {
 	FCoreUObjectDelegates::OnObjectsReplaced.AddUObject(this, &UCameraNodalOffsetAlgoPoints::OnObjectsReplaced);
@@ -705,13 +1039,12 @@ bool UCameraNodalOffsetAlgoPoints::CalculatedOptimalCameraComponentPose(
 
 	std::vector<cv::Point3f> Points3d;
 	std::vector<cv::Point2f> Points2d;
-
-	TArray<FVector2D> ImagePoints;
+	Points3d.reserve(Rows.Num());
+	Points2d.reserve(Rows.Num());
 
 	for (const TSharedPtr<FCalibrationRowData>& Row : Rows)
 	{
 		// Convert from UE coordinates to OpenCV coordinates
-
 		FTransform Transform;
 		Transform.SetIdentity();
 		Transform.SetLocation(Row->CalibratorPointData.Location);
@@ -724,9 +1057,10 @@ bool UCameraNodalOffsetAlgoPoints::CalculatedOptimalCameraComponentPose(
 			Transform.GetLocation().Y,
 			Transform.GetLocation().Z));
 
-		ImagePoints.Add(FVector2D(Row->Point2D.X, Row->Point2D.Y));
+		Points2d.push_back(cv::Point2f(
+			Row->UndistortedPoint2D.X,
+			Row->UndistortedPoint2D.Y));
 	}
-
 	// Populate camera matrix
 
 	cv::Mat CameraMatrix(3, 3, cv::DataType<double>::type);
@@ -745,19 +1079,6 @@ bool UCameraNodalOffsetAlgoPoints::CalculatedOptimalCameraComponentPose(
 	CameraMatrix.at<double>(0, 2) = 0.5;
 	CameraMatrix.at<double>(1, 2) = 0.5;
 
-	// Manually undistort the 2D image points
-	TArray<FVector2D> UndistortedPoints;
-	UndistortedPoints.AddZeroed(ImagePoints.Num());
-	DistortionRenderingUtils::UndistortImagePoints(DistortionHandler->GetDistortionDisplacementMap(), ImagePoints, UndistortedPoints);
-
-	Points2d.reserve(UndistortedPoints.Num());
-	for (FVector2D Point : UndistortedPoints)
-	{
-		Points2d.push_back(cv::Point2f(
-			Point.X,
-			Point.Y));
-	}
-
 	// Solve for camera position
 	cv::Mat Rrod = cv::Mat::zeros(3, 1, cv::DataType<double>::type); // Rotation vector in Rodrigues notation. 3x1.
 	cv::Mat Tobj = cv::Mat::zeros(3, 1, cv::DataType<double>::type); // Translation vector. 3x1.
@@ -987,6 +1308,9 @@ bool UCameraNodalOffsetAlgoPoints::GetNodalOffset(FNodalPointOffset& OutNodalOff
 		}
 	}
 
+	// Undistort the 2D points in each calibration row 
+	UndistortCalibrationRowPoints();
+
 	TArray<FSinglePoseResult> SinglePoseResults;
 	SinglePoseResults.Reserve(SamePoseRowGroups.Num());
 
@@ -1029,6 +1353,9 @@ bool UCameraNodalOffsetAlgoPoints::GetNodalOffset(FNodalPointOffset& OutNodalOff
 		return false;
 	}
 
+	// Refine the weighted average transform using a downhill solver
+	OutError = MinimizeReprojectionError(AverageTransform, SamePoseRowGroups);
+
 	// Assign output nodal offset.
 
 	OutNodalOffset.LocationOffset = AverageTransform.GetLocation();
@@ -1565,6 +1892,9 @@ bool UCameraNodalOffsetAlgoPoints::CalcCalibratorPoseForSingleCamPose(
 {
 	FTransform DesiredCameraPose;
 
+	// Undistort the 2D points in each calibration row 
+	UndistortCalibrationRowPoints();
+
 	if (!CalculatedOptimalCameraComponentPose(DesiredCameraPose, Rows, OutErrorMessage))
 	{
 		return false;
@@ -2096,6 +2426,60 @@ void UCameraNodalOffsetAlgoPoints::ClearCalibrationRows()
 	}
 }
 
+void UCameraNodalOffsetAlgoPoints::UndistortCalibrationRowPoints()
+{
+	if (CalibrationRows.Num() == 0)
+	{
+		return;
+	}
+
+	// If the data for these rows was already undistorted, early-out
+	const TSharedPtr<FCalibrationRowData>& FirstRow = CalibrationRows[0];
+	if (FirstRow->bUndistortedIsValid)
+	{
+		return;
+	}
+
+	if (!NodalOffsetTool.IsValid())
+	{
+		return;
+	}
+
+	const FCameraCalibrationStepsController* StepsController = NodalOffsetTool->GetCameraCalibrationStepsController();
+
+	if (!StepsController)
+	{
+		return;
+	}
+
+	const ULensDistortionModelHandlerBase* DistortionHandler = StepsController->GetDistortionHandler();
+	if (!DistortionHandler)
+	{
+		return;
+	}
+
+	const int32 NumPoints = CalibrationRows.Num();
+
+	TArray<FVector2D> ImagePoints;
+	ImagePoints.Reserve(NumPoints);
+
+	for (const TSharedPtr<FCalibrationRowData>& Row : CalibrationRows)
+	{
+		ImagePoints.Add(FVector2D(Row->Point2D.X, Row->Point2D.Y));
+	}
+
+	TArray<FVector2D> UndistortedPoints;
+	UndistortedPoints.AddZeroed(ImagePoints.Num());
+	DistortionRenderingUtils::UndistortImagePoints(DistortionHandler->GetDistortionDisplacementMap(), ImagePoints, UndistortedPoints);
+
+	for (int32 Index = 0; Index < NumPoints; ++Index)
+	{
+		const TSharedPtr<FCalibrationRowData>& Row = CalibrationRows[Index];
+		Row->UndistortedPoint2D = UndistortedPoints[Index];
+		Row->bUndistortedIsValid = true;
+	}
+}
+
 bool UCameraNodalOffsetAlgoPoints::GetStepsControllerAndLensFile(
 	const FCameraCalibrationStepsController** OutStepsController,
 	const ULensFile** OutLensFile) const
diff --git a/Engine/Plugins/VirtualProduction/CameraCalibration/Source/CameraCalibrationEditor/Private/Calibrators/CameraNodalOffsetAlgoPoints.h b/Engine/Plugins/VirtualProduction/CameraCalibration/Source/CameraCalibrationEditor/Private/Calibrators/CameraNodalOffsetAlgoPoints.h
--- a/Engine/Plugins/VirtualProduction/CameraCalibration/Source/CameraCalibrationEditor/Private/Calibrators/CameraNodalOffsetAlgoPoints.h
+++ b/Engine/Plugins/VirtualProduction/CameraCalibration/Source/CameraCalibrationEditor/Private/Calibrators/CameraNodalOffsetAlgoPoints.h
@@ -125,6 +125,12 @@ protected:
 		// Normalized 0~1 2d location of the identified calibrator point in the media plate.
 		FVector2D Point2D;
 
+		// Location of Point2D after it has been undistorted
+		FVector2D UndistortedPoint2D;
+
+		// True if this row has had its Point2D undistorted
+		bool bUndistortedIsValid = false;
+
 		// Holds information of the calibrator point data for this sample.
 		FCalibratorPointCache CalibratorPointData;
 
@@ -198,6 +204,9 @@ protected:
 	/** Clears the list of calibration sample rows */
 	void ClearCalibrationRows();
 
+	/* Undistort the 2D points in each calibration row using the current distortion displacement map */
+	void UndistortCalibrationRowPoints();
+
 	/** 
 	 * Retrieves by name the calibration point data of the currently selected calibrator.
 	 * 
@@ -340,4 +349,10 @@ protected:
 		const TArray<TSharedPtr<FCalibrationRowData>>& Rows,
 		FTransform& OutTransform,
 		FText& OutErrorMessage);
+
+	/* Refine the input nodal offset by running a minimization algorithm designed to minimize the sum of the reprojection errors for the camera views in the set of input pose groups*/
+	double MinimizeReprojectionError(FTransform& InOutNodalOffset, const TArray<TSharedPtr<TArray<TSharedPtr<FCalibrationRowData>>>>& SamePoseRowGroups) const;
+
+	/* Compute the reprojection error for the rows in the input PoseGroup using the input NodalOffset */
+	double ComputeReprojectionError(const FTransform& NodalOffset, const TArray<TSharedPtr<FCalibrationRowData>>& PoseGroup) const;
 };
