Fix crash when an overlay is enabled in calibration tools and GC kicks in

diff --git a/Engine/Plugins/VirtualProduction/CameraCalibration/Source/CameraCalibrationEditor/Private/AssetEditor/SCameraCalibrationSteps.cpp b/Engine/Plugins/VirtualProduction/CameraCalibration/Source/CameraCalibrationEditor/Private/AssetEditor/SCameraCalibrationSteps.cpp
--- a/Engine/Plugins/VirtualProduction/CameraCalibration/Source/CameraCalibrationEditor/Private/AssetEditor/SCameraCalibrationSteps.cpp
+++ b/Engine/Plugins/VirtualProduction/CameraCalibration/Source/CameraCalibrationEditor/Private/AssetEditor/SCameraCalibrationSteps.cpp
@@ -34,6 +34,21 @@
 
 #define LOCTEXT_NAMESPACE "CameraCalibrationSteps"
 
+void SCameraCalibrationSteps::AddReferencedObjects(FReferenceCollector& Collector)
+{
+ 	if (CurrentOverlayMID)
+ 	{
+		Collector.AddReferencedObject(CurrentOverlayMID);
+	}
+
+	for (TPair<FName, TObjectPtr<UMaterialInstanceDynamic>> Pair : OverlayMIDs)
+	{
+		if (Pair.Value)
+		{
+			Collector.AddReferencedObject(Pair.Value);
+		}
+	}
+}
 
 void SCameraCalibrationSteps::Construct(const FArguments& InArgs, TWeakPtr<FCameraCalibrationStepsController> InCalibrationStepsController)
 {
diff --git a/Engine/Plugins/VirtualProduction/CameraCalibration/Source/CameraCalibrationEditor/Private/AssetEditor/SCameraCalibrationSteps.h b/Engine/Plugins/VirtualProduction/CameraCalibration/Source/CameraCalibrationEditor/Private/AssetEditor/SCameraCalibrationSteps.h
--- a/Engine/Plugins/VirtualProduction/CameraCalibration/Source/CameraCalibrationEditor/Private/AssetEditor/SCameraCalibrationSteps.h
+++ b/Engine/Plugins/VirtualProduction/CameraCalibration/Source/CameraCalibrationEditor/Private/AssetEditor/SCameraCalibrationSteps.h
@@ -22,7 +22,7 @@ struct FArguments;
  * UI for the nodal offset calibration.
  * It also holds the UI given by the selected nodal offset algorithm.
  */
-class SCameraCalibrationSteps : public SCompoundWidget
+class SCameraCalibrationSteps : public SCompoundWidget, public FGCObject
 {
 	SLATE_BEGIN_ARGS(SCameraCalibrationSteps) {}
 	SLATE_END_ARGS()
@@ -30,6 +30,15 @@ class SCameraCalibrationSteps : public SCompoundWidget
 public:
 	void Construct(const FArguments& InArgs, TWeakPtr<FCameraCalibrationStepsController> InCalibrationStepsController);
 
+protected:
+	//~ Begin FGCObject interface
+	virtual void AddReferencedObjects(FReferenceCollector& Collector) override;
+	virtual FString GetReferencerName() const override
+	{
+		return TEXT("SCameraCalibrationSteps");
+	}
+	//~ End FGCObject interface
+
 private:
 
 	/** Builds the UI used to pick the camera used for the CG layer of the comp */
