Adding missing include from last commit (to be on the safe side)

diff --git a/Engine/Plugins/VirtualProduction/CameraCalibration/Source/CameraCalibrationEditor/Private/AssetEditor/SCameraCalibrationSteps.cpp b/Engine/Plugins/VirtualProduction/CameraCalibration/Source/CameraCalibrationEditor/Private/AssetEditor/SCameraCalibrationSteps.cpp
--- a/Engine/Plugins/VirtualProduction/CameraCalibration/Source/CameraCalibrationEditor/Private/AssetEditor/SCameraCalibrationSteps.cpp
+++ b/Engine/Plugins/VirtualProduction/CameraCalibration/Source/CameraCalibrationEditor/Private/AssetEditor/SCameraCalibrationSteps.cpp
@@ -41,7 +41,7 @@ void SCameraCalibrationSteps::AddReferencedObjects(FReferenceCollector& Collecto
 		Collector.AddReferencedObject(CurrentOverlayMID);
 	}
 
-	for (TPair<FName, TObjectPtr<UMaterialInstanceDynamic>> Pair : OverlayMIDs)
+	for (TPair<FName, TObjectPtr<UMaterialInstanceDynamic>>& Pair : OverlayMIDs)
 	{
 		if (Pair.Value)
 		{
diff --git a/Engine/Plugins/VirtualProduction/CameraCalibration/Source/CameraCalibrationEditor/Private/AssetEditor/SCameraCalibrationSteps.h b/Engine/Plugins/VirtualProduction/CameraCalibration/Source/CameraCalibrationEditor/Private/AssetEditor/SCameraCalibrationSteps.h
--- a/Engine/Plugins/VirtualProduction/CameraCalibration/Source/CameraCalibrationEditor/Private/AssetEditor/SCameraCalibrationSteps.h
+++ b/Engine/Plugins/VirtualProduction/CameraCalibration/Source/CameraCalibrationEditor/Private/AssetEditor/SCameraCalibrationSteps.h
@@ -6,6 +6,7 @@
 #include "CameraCalibrationStepsController.h"
 #include "Widgets/DeclarativeSyntaxSupport.h"
 #include "Widgets/SCompoundWidget.h"
+#include "UObject/GCObject.h"
 #include "UObject/StrongObjectPtr.h"
 
 class FString;
