Moved more OpenCV convenience functions into OpenCVHelper

diff --git a/Engine/Plugins/Runtime/OpenCV/Source/OpenCVHelper/Private/OpenCVHelper.cpp b/Engine/Plugins/Runtime/OpenCV/Source/OpenCVHelper/Private/OpenCVHelper.cpp
--- a/Engine/Plugins/Runtime/OpenCV/Source/OpenCVHelper/Private/OpenCVHelper.cpp
+++ b/Engine/Plugins/Runtime/OpenCV/Source/OpenCVHelper/Private/OpenCVHelper.cpp
@@ -11,7 +11,40 @@ OPENCV_INCLUDES_START
 #undef check 
 #include "opencv2/opencv.hpp"
 OPENCV_INCLUDES_END
+#endif	// WITH_OPENCV
 
+void FOpenCVHelper::ConvertCoordinateSystem(FTransform& Transform, const EAxis SrcXInDstAxis, const EAxis SrcYInDstAxis, const EAxis SrcZInDstAxis)
+{
+	// Unreal Engine:
+	//   Front : X
+	//   Right : Y
+	//   Up    : Z
+	//
+	// OpenCV:
+	//   Front : Z
+	//   Right : X
+	//   Up    : Yn
+
+	FMatrix M12 = FMatrix::Identity;
+
+	M12.SetColumn(0, UnitVectorFromAxisEnum(SrcXInDstAxis));
+	M12.SetColumn(1, UnitVectorFromAxisEnum(SrcYInDstAxis));
+	M12.SetColumn(2, UnitVectorFromAxisEnum(SrcZInDstAxis));
+
+	Transform.SetFromMatrix(M12.GetTransposed() * Transform.ToMatrixWithScale() * M12);
+}
+
+void FOpenCVHelper::ConvertUnrealToOpenCV(FTransform& Transform)
+{
+	ConvertCoordinateSystem(Transform, EAxis::Y, EAxis::Zn, EAxis::X);
+}
+
+void FOpenCVHelper::ConvertOpenCVToUnreal(FTransform& Transform)
+{
+	ConvertCoordinateSystem(Transform, EAxis::Z, EAxis::X, EAxis::Yn);
+}
+
+#if WITH_OPENCV
 cv::Mat FOpenCVLensDistortionParametersBase::ConvertToOpenCVDistortionCoefficients() const
 {
 	if (bUseFisheyeModel)
diff --git a/Engine/Plugins/Runtime/OpenCV/Source/OpenCVHelper/Public/OpenCVHelper.h b/Engine/Plugins/Runtime/OpenCV/Source/OpenCVHelper/Public/OpenCVHelper.h
--- a/Engine/Plugins/Runtime/OpenCV/Source/OpenCVHelper/Public/OpenCVHelper.h
+++ b/Engine/Plugins/Runtime/OpenCV/Source/OpenCVHelper/Public/OpenCVHelper.h
@@ -3,6 +3,9 @@
 #pragma once
 #include "CoreMinimal.h"
 #include <vector>
+
+#include <type_traits>
+
 /*
 Like many third party headers, OpenCV headers require some care when importing.
 
@@ -55,11 +58,46 @@ namespace cv
 	typedef Point_<float> Point2f;
 	typedef Point3_<float> Point3f;
 };
+#endif	// WITH_OPENCV
+
 
 class OPENCVHELPER_API FOpenCVHelper
 {
 public:
+	/** Enumeration to specify any cartesian axis in positive or negative directions */
+	enum class EAxis
+	{
+		X, Y, Z,
+		Xn, Yn, Zn,
+	};
+
+	// These axes must match the order in which they are declared in EAxis
+	inline static const TArray<FVector> UnitVectors =
+	{
+		{  1,  0,  0 }, //  X
+		{  0,  1,  0 }, //  Y
+		{  0,  0,  1 }, //  Z
+		{ -1,  0,  0 }, // -X
+		{  0, -1,  0 }, // -Y
+		{  0,  0, -1 }, // -Z
+	};
+
+	static const FVector& UnitVectorFromAxisEnum(const EAxis Axis)
+	{
+		return UnitVectors[std::underlying_type_t<EAxis>(Axis)];
+	};
 
+	/** Converts in-place the coordinate system of the given FTransform by specifying the source axes in terms of the destionation axes */
+	static void ConvertCoordinateSystem(FTransform& Transform, const EAxis DstXInSrcAxis, const EAxis DstYInSrcAxis, const EAxis DstZInSrcAxis);
+
+	/** Converts in-place an FTransform in Unreal coordinates to OpenCV coordinates */
+	static void ConvertUnrealToOpenCV(FTransform& Transform);
+
+	/** Converts in-place an FTransform in OpenCV coordinates to Unreal coordinates */
+	static void ConvertOpenCVToUnreal(FTransform& Transform);
+
+#if WITH_OPENCV
+public:
 	/**
 	 * Creates a Texture from the given Mat, if its properties (e.g. pixel format) are supported.
 	 * 
@@ -73,10 +111,9 @@ public:
 	static UTexture2D* TextureFromCvMat(cv::Mat& Mat, UTexture2D* InTexture);
 
 	static double ComputeReprojectionError(const FTransform& CameraPose, const cv::Mat& CameraIntrinsicMatrix, const std::vector<cv::Point3f>& Points3d, const std::vector<cv::Point2f>& Points2d);
+#endif	// WITH_OPENCV
 };
 
-#endif //WITH_OPENCV
-
 /**
  * Mathematic camera model for lens distortion/undistortion.
  * Camera matrix =
diff --git a/Engine/Plugins/VirtualProduction/CameraCalibration/Source/CameraCalibrationEditor/Private/Calibrators/CameraNodalOffsetAlgoPoints.cpp b/Engine/Plugins/VirtualProduction/CameraCalibration/Source/CameraCalibrationEditor/Private/Calibrators/CameraNodalOffsetAlgoPoints.cpp
--- a/Engine/Plugins/VirtualProduction/CameraCalibration/Source/CameraCalibrationEditor/Private/Calibrators/CameraNodalOffsetAlgoPoints.cpp
+++ b/Engine/Plugins/VirtualProduction/CameraCalibration/Source/CameraCalibrationEditor/Private/Calibrators/CameraNodalOffsetAlgoPoints.cpp
@@ -21,6 +21,7 @@
 #include "Math/Vector.h"
 #include "Misc/FileHelper.h"
 #include "Misc/MessageDialog.h"
+#include "OpenCVHelper.h"
 #include "PropertyCustomizationHelpers.h"
 #include "UI/SFilterableActorPicker.h"
 #include "ScopedTransaction.h"
@@ -32,7 +33,6 @@
 #include <vector>
 
 #if WITH_OPENCV
-#include "OpenCVHelper.h"
 OPENCV_INCLUDES_START
 #undef check 
 #include "opencv2/opencv.hpp"
@@ -239,7 +239,7 @@ public:
 			FTransform OptimalCameraPose = NodalOffsetCandidate * ExistingNodalOffsetInverse * CameraView.CameraPose;
 
 			// Convert the optimal camera pose to opencv's coordinate system
-			FCameraCalibrationUtils::ConvertUnrealToOpenCV(OptimalCameraPose);
+			FOpenCVHelper::ConvertUnrealToOpenCV(OptimalCameraPose);
 
 			// Compute the reprojection error for this view and add it to the running total
 			double ViewError = FOpenCVHelper::ComputeReprojectionError(OptimalCameraPose, CameraMatrix, CameraView.Points3d, CameraView.Points2d);
@@ -394,7 +394,7 @@ double UCameraNodalOffsetAlgoPoints::MinimizeReprojectionError(FTransform& InOut
 			Transform.SetIdentity();
 			Transform.SetLocation(Row->CalibratorPointData.Location);
 
-			FCameraCalibrationUtils::ConvertUnrealToOpenCV(Transform);
+			FOpenCVHelper::ConvertUnrealToOpenCV(Transform);
 
 			// Calibrator 3d points
 			Points3d.push_back(cv::Point3f(
@@ -466,7 +466,7 @@ double UCameraNodalOffsetAlgoPoints::ComputeReprojectionError(const FTransform&
 		Transform.SetIdentity();
 		Transform.SetLocation(Row->CalibratorPointData.Location);
 
-		FCameraCalibrationUtils::ConvertUnrealToOpenCV(Transform);
+		FOpenCVHelper::ConvertUnrealToOpenCV(Transform);
 
 		// Calibrator 3d points
 		Points3d.push_back(cv::Point3f(
@@ -514,7 +514,7 @@ double UCameraNodalOffsetAlgoPoints::ComputeReprojectionError(const FTransform&
 	FTransform OptimalCameraPose = NodalOffset * ExistingOffset.Inverse() * CameraPose;
 
 	// Convert the optimal camera pose to opencv's coordinate system
-	FCameraCalibrationUtils::ConvertUnrealToOpenCV(OptimalCameraPose);
+	FOpenCVHelper::ConvertUnrealToOpenCV(OptimalCameraPose);
 
 	// Compute the reprojection error for this view
 	return FOpenCVHelper::ComputeReprojectionError(OptimalCameraPose, CameraMatrix, Points3d, Points2d);
@@ -1049,7 +1049,7 @@ bool UCameraNodalOffsetAlgoPoints::CalculatedOptimalCameraComponentPose(
 		Transform.SetIdentity();
 		Transform.SetLocation(Row->CalibratorPointData.Location);
 
-		FCameraCalibrationUtils::ConvertUnrealToOpenCV(Transform);
+		FOpenCVHelper::ConvertUnrealToOpenCV(Transform);
 
 		// Calibrator 3d points
 		Points3d.push_back(cv::Point3f(
@@ -1139,7 +1139,7 @@ bool UCameraNodalOffsetAlgoPoints::CalculatedOptimalCameraComponentPose(
 	M.M[3][2] = Tcam.at<double>(2);
 
 	OutDesiredCameraTransform.SetFromMatrix(M);
-	FCameraCalibrationUtils::ConvertOpenCVToUnreal(OutDesiredCameraTransform);
+	FOpenCVHelper::ConvertOpenCVToUnreal(OutDesiredCameraTransform);
 
 	return true;
 
diff --git a/Engine/Plugins/VirtualProduction/CameraCalibration/Source/CameraCalibrationEditor/Private/LensFileExchangeFormat.cpp b/Engine/Plugins/VirtualProduction/CameraCalibration/Source/CameraCalibrationEditor/Private/LensFileExchangeFormat.cpp
--- a/Engine/Plugins/VirtualProduction/CameraCalibration/Source/CameraCalibrationEditor/Private/LensFileExchangeFormat.cpp
+++ b/Engine/Plugins/VirtualProduction/CameraCalibration/Source/CameraCalibrationEditor/Private/LensFileExchangeFormat.cpp
@@ -3,14 +3,13 @@
 #include "LensFileExchangeFormat.h"
 
 #include "CameraCalibrationEditorLog.h"
-#include "CameraCalibrationUtils.h"
 #include "JsonObjectConverter.h"
 #include "LensData.h"
 #include "LensFile.h"
 #include "Misc/FeedbackContext.h"
 #include "Models/LensModel.h"
 #include "Models/SphericalLensModel.h"
-
+#include "OpenCVHelper.h"
 
 #define LOCTEXT_NAMESPACE "LensFileExchange"
 
@@ -517,7 +516,7 @@ void FLensFileExchange::ExtractNodalOffsetTable(const ULensFile* LensFile)
 			if (NodalOffsetFocusPoint.GetPoint(Zoom, NodalOffsetInfo))
 			{
 				FTransform NodalOffsetTransform{ NodalOffsetInfo.RotationOffset, NodalOffsetInfo.LocationOffset };
-				FCameraCalibrationUtils::ConvertUnrealToOpenCV(NodalOffsetTransform);
+				FOpenCVHelper::ConvertUnrealToOpenCV(NodalOffsetTransform);
 
 				const FQuat RotationOffset = NodalOffsetTransform.GetRotation();
 				const FVector LocationOffset = NodalOffsetTransform.GetTranslation();
@@ -823,7 +822,7 @@ void FLensFileExchange::PopulateNodalOffsetTable(ULensFile& OutLensFile, const F
 
 			if (Metadata.NodalOffsetCoordinateSystem == ENodalOffsetCoordinateSystem::OpenCV)
 			{
-				FCameraCalibrationUtils::ConvertUnrealToOpenCV(NodalOffsetTransform);
+				FOpenCVHelper::ConvertUnrealToOpenCV(NodalOffsetTransform);
 			}
 
 			FNodalPointOffset NodalPointOffset;
diff --git a/Engine/Plugins/VirtualProduction/CameraCalibrationCore/Source/CameraCalibrationCore/Private/CameraCalibrationUtils.cpp b/Engine/Plugins/VirtualProduction/CameraCalibrationCore/Source/CameraCalibrationCore/Private/CameraCalibrationUtils.cpp
--- a/Engine/Plugins/VirtualProduction/CameraCalibrationCore/Source/CameraCalibrationCore/Private/CameraCalibrationUtils.cpp
+++ b/Engine/Plugins/VirtualProduction/CameraCalibrationCore/Source/CameraCalibrationCore/Private/CameraCalibrationUtils.cpp
@@ -3,65 +3,8 @@
 
 #include "CameraCalibrationUtils.h"
 
-#include "Engine/Texture2D.h"
 #include "Math/UnrealMathUtility.h"
 
-#include <type_traits>
-
-namespace CameraCalibrationUtils
-{
-	using EAxis = FCameraCalibrationUtils::EAxis;
-
-	// These axes must match the order in which they are declared in EAxis
-	static const TArray<FVector> UnitVectors =
-	{
-		{  1,  0,  0 }, //  X
-		{  0,  1,  0 }, //  Y
-		{  0,  0,  1 }, //  Z
-		{ -1,  0,  0 }, // -X
-		{  0, -1,  0 }, // -Y
-		{  0,  0, -1 }, // -Z
-	};
-
-	static const FVector& UnitVectorFromAxisEnum(EAxis Axis)
-	{
-		return UnitVectors[std::underlying_type_t<EAxis>(Axis)];
-	};
-}
-
-void FCameraCalibrationUtils::ConvertCoordinateSystem(FTransform& Transform, EAxis SrcXInDstAxis, EAxis SrcYInDstAxis, EAxis SrcZInDstAxis)
-{
-	// Unreal Engine:
-	//   Front : X
-	//   Right : Y
-	//   Up    : Z
-	//
-	// OpenCV:
-	//   Front : Z
-	//   Right : X
-	//   Up    : Yn
-
-	using namespace CameraCalibrationUtils;
-
-	FMatrix M12 = FMatrix::Identity;
-
-	M12.SetColumn(0, UnitVectorFromAxisEnum(SrcXInDstAxis));
-	M12.SetColumn(1, UnitVectorFromAxisEnum(SrcYInDstAxis));
-	M12.SetColumn(2, UnitVectorFromAxisEnum(SrcZInDstAxis));
-
-	Transform.SetFromMatrix(M12.GetTransposed() * Transform.ToMatrixWithScale() * M12);
-}
-
-void FCameraCalibrationUtils::ConvertUnrealToOpenCV(FTransform& Transform)
-{
-	ConvertCoordinateSystem(Transform, EAxis::Y, EAxis::Zn, EAxis::X);
-}
-
-void FCameraCalibrationUtils::ConvertOpenCVToUnreal(FTransform& Transform)
-{
-	ConvertCoordinateSystem(Transform, EAxis::Z, EAxis::X, EAxis::Yn);
-}
-
 bool FCameraCalibrationUtils::IsNearlyEqual(const FTransform& A, const FTransform& B, float MaxLocationDelta, float MaxAngleDeltaDegrees)
 {
 	// Location check
diff --git a/Engine/Plugins/VirtualProduction/CameraCalibrationCore/Source/CameraCalibrationCore/Public/CameraCalibrationUtils.h b/Engine/Plugins/VirtualProduction/CameraCalibrationCore/Source/CameraCalibrationCore/Public/CameraCalibrationUtils.h
--- a/Engine/Plugins/VirtualProduction/CameraCalibrationCore/Source/CameraCalibrationCore/Public/CameraCalibrationUtils.h
+++ b/Engine/Plugins/VirtualProduction/CameraCalibrationCore/Source/CameraCalibrationCore/Public/CameraCalibrationUtils.h
@@ -4,34 +4,12 @@
 
 #include "CoreMinimal.h"
 
-
-class UTexture2D;
-
 /**
  * Helper class for commonly used functions for camera calibration.
  */
 class CAMERACALIBRATIONCORE_API FCameraCalibrationUtils
 {
 public:
-
-	/** Enumeration to specify any cartesian axis in positive or negative directions */
-	enum class EAxis
-	{
-		X, Y, Z,
-		Xn, Yn, Zn,
-	};
-
-public:
-
-	/** Converts in-place the coordinate system of the given FTransform by specifying the source axes in terms of the destionation axes */
-	static void ConvertCoordinateSystem(FTransform& Transform, EAxis DstXInSrcAxis, EAxis DstYInSrcAxis, EAxis DstZInSrcAxis);
-
-	/** Converts in-place an FTransform in Unreal coordinates to OpenCV coordinates */
-	static void ConvertUnrealToOpenCV(FTransform& Transform);
-
-	/** Converts in-place an FTransform in OpenCV coordinates to Unreal coordinates */
-	static void ConvertOpenCVToUnreal(FTransform& Transform);
-
 	/** Compares two transforms and returns true if they are nearly equal in distance and angle */
 	static bool IsNearlyEqual(const FTransform& A, const FTransform& B, float MaxLocationDelta = 2.0f, float MaxAngleDegrees = 2.0f);
 };
