Nodal offset calibration algorithm for finding the optical axis of the lens and the entrance pupil position that lies along that axis

diff --git a/Engine/Plugins/VirtualProduction/CameraCalibration/Source/CameraCalibrationEditor/Private/AssetEditor/NodalOffsetTool.cpp b/Engine/Plugins/VirtualProduction/CameraCalibration/Source/CameraCalibrationEditor/Private/AssetEditor/NodalOffsetTool.cpp
--- a/Engine/Plugins/VirtualProduction/CameraCalibration/Source/CameraCalibrationEditor/Private/AssetEditor/NodalOffsetTool.cpp
+++ b/Engine/Plugins/VirtualProduction/CameraCalibration/Source/CameraCalibrationEditor/Private/AssetEditor/NodalOffsetTool.cpp
@@ -66,6 +66,21 @@ bool UNodalOffsetTool::OnViewportClicked(const FGeometry& MyGeometry, const FPoi
 	return NodalOffsetAlgo->OnViewportClicked(MyGeometry, MouseEvent);
 }
 
+bool UNodalOffsetTool::OnViewportInputKey(const FKey& InKey, const EInputEvent& InEvent)
+{
+	if (!bIsActive)
+	{
+		return false;
+	}
+
+	if (!NodalOffsetAlgo)
+	{
+		return false;
+	}
+
+	return NodalOffsetAlgo->OnViewportInputKey(InKey, InEvent);
+}
+
 TSharedRef<SWidget> UNodalOffsetTool::BuildUI()
 {
 	return SNew(SNodalOffsetToolPanel, this);
@@ -128,7 +143,8 @@ void UNodalOffsetTool::SetNodalOffsetAlgo(const FName& AlgoName)
 	NodalOffsetAlgo = NewObject<UCameraNodalOffsetAlgo>(
 		GetTransientPackage(),
 		AlgoClass,
-		MakeUniqueObjectName(GetTransientPackage(), AlgoClass));
+		MakeUniqueObjectName(GetTransientPackage(), AlgoClass),
+		RF_Transactional);
 
 	if (NodalOffsetAlgo)
 	{
diff --git a/Engine/Plugins/VirtualProduction/CameraCalibration/Source/CameraCalibrationEditor/Private/AssetEditor/NodalOffsetTool.h b/Engine/Plugins/VirtualProduction/CameraCalibration/Source/CameraCalibrationEditor/Private/AssetEditor/NodalOffsetTool.h
--- a/Engine/Plugins/VirtualProduction/CameraCalibration/Source/CameraCalibrationEditor/Private/AssetEditor/NodalOffsetTool.h
+++ b/Engine/Plugins/VirtualProduction/CameraCalibration/Source/CameraCalibrationEditor/Private/AssetEditor/NodalOffsetTool.h
@@ -33,6 +33,7 @@ public:
 	virtual void Shutdown() override;
 	virtual void Tick(float DeltaTime) override;
 	virtual bool OnViewportClicked(const FGeometry& MyGeometry, const FPointerEvent& MouseEvent) override;
+	virtual bool OnViewportInputKey(const FKey& InKey, const EInputEvent& InEvent) override;
 	virtual TSharedRef<SWidget> BuildUI() override;
 	virtual FName FriendlyName() const  override { return TEXT("Nodal Offset"); };
 	virtual bool DependsOnStep(UCameraCalibrationStep* Step) const override;
diff --git a/Engine/Plugins/VirtualProduction/CameraCalibration/Source/CameraCalibrationEditor/Private/AssetEditor/Viewport/SimulcamEditorViewportClient.cpp b/Engine/Plugins/VirtualProduction/CameraCalibration/Source/CameraCalibrationEditor/Private/AssetEditor/Viewport/SimulcamEditorViewportClient.cpp
--- a/Engine/Plugins/VirtualProduction/CameraCalibration/Source/CameraCalibrationEditor/Private/AssetEditor/Viewport/SimulcamEditorViewportClient.cpp
+++ b/Engine/Plugins/VirtualProduction/CameraCalibration/Source/CameraCalibrationEditor/Private/AssetEditor/Viewport/SimulcamEditorViewportClient.cpp
@@ -188,6 +188,10 @@ bool FSimulcamEditorViewportClient::InputKey(FViewport* Viewport, int32 Controll
 
 		return SimulcamViewportWeakPtr.Pin()->OnViewportInputKey(Key, Event);
 	}
+	else if (Event == IE_Released)
+	{
+		return SimulcamViewportWeakPtr.Pin()->OnViewportInputKey(Key, Event);
+	}
 	else if (Event == IE_Repeat)
 	{
 		return SimulcamViewportWeakPtr.Pin()->OnViewportInputKey(Key, Event);
diff --git a/Engine/Plugins/VirtualProduction/CameraCalibration/Source/CameraCalibrationEditor/Private/Calibrators/CameraNodalOffsetAlgoOpticalAxis.cpp b/Engine/Plugins/VirtualProduction/CameraCalibration/Source/CameraCalibrationEditor/Private/Calibrators/CameraNodalOffsetAlgoOpticalAxis.cpp
new file mode 100644
--- /dev/null
+++ b/Engine/Plugins/VirtualProduction/CameraCalibration/Source/CameraCalibrationEditor/Private/Calibrators/CameraNodalOffsetAlgoOpticalAxis.cpp
@@ -0,0 +1,567 @@
+// Copyright Epic Games, Inc. All Rights Reserved.
+
+#include "CameraNodalOffsetAlgoOpticalAxis.h"
+
+#include "AssetEditor/CameraCalibrationStepsController.h"
+#include "AssetEditor/NodalOffsetTool.h"
+#include "CameraCalibrationEditorLog.h"
+#include "CameraCalibrationSettings.h"
+#include "CameraCalibrationUtils.h"
+#include "Kismet/KismetMathLibrary.h"
+#include "LensDistortionModelHandlerBase.h"
+#include "Misc/MessageDialog.h"
+#include "OpenCVHelper.h"
+#include "OverlayRendering.h"
+#include "ScopedTransaction.h"
+#include "UI/CameraCalibrationWidgetHelpers.h"
+#include "Widgets/Input/SButton.h"
+#include "Widgets/Input/SNumericEntryBox.h"
+
+#include <vector>
+
+#if WITH_OPENCV
+OPENCV_INCLUDES_START
+#undef check
+#include "opencv2/opencv.hpp"
+#include "opencv2/calib3d.hpp"
+OPENCV_INCLUDES_END
+#endif
+
+
+#define LOCTEXT_NAMESPACE "CameraNodalOffsetAlgoOpticalAxis"
+
+
+void UCameraNodalOffsetAlgoOpticalAxis::Initialize(UNodalOffsetTool* InNodalOffsetTool)
+{
+	Super::Initialize(InNodalOffsetTool);
+
+	OpticalAxis.Reset();
+
+	const FCameraCalibrationStepsController* StepsController = NodalOffsetTool->GetCameraCalibrationStepsController();
+
+	if (!StepsController)
+	{
+		return;
+	}
+
+	// Create the render target for the crosshair overlay
+	const FIntPoint OverlaySize = StepsController->GetCompRenderTargetSize();
+
+	OverlayTexture = NewObject<UTextureRenderTarget2D>(GetTransientPackage(), MakeUniqueObjectName(GetTransientPackage(), UTextureRenderTarget2D::StaticClass()));
+
+	OverlayTexture->RenderTargetFormat = RTF_RGBA16f;
+	OverlayTexture->ClearColor = FLinearColor::Transparent;
+	OverlayTexture->bAutoGenerateMips = false;
+	OverlayTexture->InitAutoFormat(OverlaySize.X, OverlaySize.Y);
+	OverlayTexture->UpdateResourceImmediate(true);
+
+	if (UMaterialInstanceDynamic* OverlayMID = NodalOffsetTool->GetOverlayMID())
+	{
+		OverlayMID->SetTextureParameterValue(FName(TEXT("OutputTexture")), OverlayTexture);
+	}
+
+	// Initialize sensitivity options for the combobox
+	SensitivityOptions.Add(MakeShared<float>(0.001f));
+	SensitivityOptions.Add(MakeShared<float>(0.005f));
+	SensitivityOptions.Add(MakeShared<float>(0.01f));
+	SensitivityOptions.Add(MakeShared<float>(0.05f));
+	SensitivityOptions.Add(MakeShared<float>(0.1f));
+	SensitivityOptions.Add(MakeShared<float>(0.5f));
+}
+
+void UCameraNodalOffsetAlgoOpticalAxis::Tick(float DeltaTime)
+{
+	Super::Tick(DeltaTime);
+
+	if (!NodalOffsetTool.IsValid())
+	{
+		return;
+	}
+
+	FCameraCalibrationStepsController* StepsController = NodalOffsetTool->GetCameraCalibrationStepsController();
+
+	if (!StepsController)
+	{
+		return;
+	}
+
+	const ULensDistortionModelHandlerBase* DistortionHandler = StepsController->GetDistortionHandler();
+
+	if (!DistortionHandler)
+	{
+		return;
+	}
+
+	const FLensDistortionState DistortionState = DistortionHandler->GetCurrentDistortionState();
+
+	// Inform the user if the focus or zoom has changed since the last tick
+	if (OpticalAxis.IsSet() && !bFocusZoomWarningIssued)
+	{
+		if (const FLensFileEvalData* EvalData = StepsController->GetLensFileEvalData())
+		{
+			const float InputTolerance = GetDefault<UCameraCalibrationSettings>()->GetCalibrationInputTolerance();
+			if (!FMath::IsNearlyEqual(CachedFocus, EvalData->Input.Focus, InputTolerance) || !FMath::IsNearlyEqual(CachedZoom, EvalData->Input.Zoom, InputTolerance))
+			{
+				FFormatOrderedArguments Arguments;
+				Arguments.Add(FText::FromString(FString::Printf(TEXT("%.2f"), CachedFocus)));
+				Arguments.Add(FText::FromString(FString::Printf(TEXT("%.2f"), CachedZoom)));
+
+				FText Message = FText::Format(LOCTEXT("Focus or Zoom has changed",
+					"The focus or zoom input has changed. You may continue to modify the entrance pupil position, "
+					"but note that those modifications will affect the nodal offset at the previous values of "
+					"Focus = {0}, Zoom = {1}. You may calibrate at the new focus and zoom values by capturing new "
+					"points and clicking 'Add To Nodal Offset Calibration'."
+				), Arguments);
+
+				FMessageDialog::Open(EAppMsgType::Ok, Message);
+
+				bFocusZoomWarningIssued = true;
+			}
+		}
+	}
+
+	// Update the crosshair overlay if the principal point has changed
+	if (LastPrincipalPoint != DistortionState.ImageCenter.PrincipalPoint)
+	{
+		FCrosshairOverlayParams OverlayParams;
+		OverlayParams.PrincipalPoint = DistortionState.ImageCenter.PrincipalPoint;
+		OverlayRendering::DrawCrosshairOverlay(OverlayTexture, OverlayParams);
+		StepsController->RefreshOverlay();
+
+		LastPrincipalPoint = DistortionState.ImageCenter.PrincipalPoint;
+	}
+}
+
+bool UCameraNodalOffsetAlgoOpticalAxis::OnViewportInputKey(const FKey& InKey, const EInputEvent& InEvent)
+{
+	if (!OpticalAxis.IsSet())
+	{
+		return false;
+	}
+
+	bool bHandled = false;
+
+	// Arrow Key handling
+	if (InEvent == EInputEvent::IE_Pressed || InEvent == EInputEvent::IE_Repeat)
+	{
+		if (InKey == EKeys::Up)
+		{
+			EntrancePupilPosition -= (AdjustmentIncrement);
+			bHandled = true;
+		}
+		else if (InKey == EKeys::Down)
+		{
+			EntrancePupilPosition += (AdjustmentIncrement);
+			bHandled = true;
+		}
+
+		if (bHandled)
+		{
+			UpdateNodalOffset();
+		}
+	}
+	else if (InEvent == EInputEvent::IE_Released)
+	{
+		// Only trigger a transaction when the key is released
+		if (InKey == EKeys::Up || InKey == EKeys::Down)
+		{
+			SaveUpdatedNodalOffset();
+			bHandled = true;
+		}
+	}
+
+	return bHandled;
+}
+
+bool UCameraNodalOffsetAlgoOpticalAxis::GetNodalOffset(FNodalPointOffset& OutNodalOffset, float& OutFocus, float& OutZoom, float& OutError, FText& OutErrorMessage)
+{
+	if (!NodalOffsetTool.IsValid())
+	{
+		return false;
+	}
+
+	const FCameraCalibrationStepsController* StepsController = NodalOffsetTool->GetCameraCalibrationStepsController();
+
+	if (!StepsController)
+	{
+		return false;
+	}
+
+	const ULensFile* LensFile = StepsController->GetLensFile();
+
+	if (!LensFile)
+	{
+		return false;
+	}
+
+	const ULensDistortionModelHandlerBase* DistortionHandler = StepsController->GetDistortionHandler();
+
+	if (!DistortionHandler)
+	{
+		return false;
+	}
+
+	// Validate that there are enough points to compute an optical axis
+ 	if (CalibrationRows.Num() < 2)
+ 	{
+ 		OutErrorMessage = LOCTEXT("AtLeastTwoPoints", "At least two points are required");
+ 		return false;
+ 	}
+
+	// Validate that all rows have approximately the same camera pose
+	const TSharedPtr<FCalibrationRowData>& FirstRow = CalibrationRows[0];
+	const FTransform CameraPose = FirstRow->CameraData.Pose;
+	CachedInverseCameraPose = CameraPose.Inverse();
+
+	for (const TSharedPtr<FCalibrationRowData>& Row : CalibrationRows)
+	{
+		if (!FCameraCalibrationUtils::IsNearlyEqual(CameraPose, Row->CameraData.Pose))
+		{
+			OutErrorMessage = LOCTEXT("CameraMoved", "Camera moved too much between samples.");
+			return false;
+		}
+	}
+
+	// Cache the evaluated focus and zoom so that modifications to the nodal offset table will update the same point
+	OutFocus = CachedFocus = FirstRow->CameraData.LensFileEvalData.Input.Focus;
+	OutZoom = CachedZoom = FirstRow->CameraData.LensFileEvalData.Input.Zoom;
+
+#if WITH_OPENCV
+	// Gather the 3D points from each of the calibration rows
+	std::vector<cv::Point3f> Points3d;
+	Points3d.reserve(CalibrationRows.Num());
+
+	for (const TSharedPtr<FCalibrationRowData>& Row : CalibrationRows)
+	{
+		Points3d.push_back(cv::Point3f(
+			Row->CalibratorPointData.Location.X,
+			Row->CalibratorPointData.Location.Y,
+			Row->CalibratorPointData.Location.Z
+		));
+	}
+
+	// Find a best fit line between the 3D points, producing a line (the optical axis) and a point on that line
+	cv::Vec6f LineAndPoint;
+	cv::fitLine(Points3d, LineAndPoint, cv::DIST_L2, 0, 0.01, 0.01);
+
+	OpticalAxis = FVector(LineAndPoint[0], LineAndPoint[1], LineAndPoint[2]);
+	PointAlongAxis = FVector(LineAndPoint[3], LineAndPoint[4], LineAndPoint[5]);
+
+	// Find a look at rotation for the camera
+	const FVector Point0 = PointAlongAxis + (0.0f * OpticalAxis.GetValue());
+	const FVector Point1 = PointAlongAxis + (1.0f * OpticalAxis.GetValue());
+
+	LookAtRotation = UKismetMathLibrary::FindLookAtRotation(Point1, Point0);
+
+	// Cache the existing nodal offset, if one exists
+	CachedNodalOffset = FTransform::Identity;
+
+	if (FirstRow->CameraData.LensFileEvalData.NodalOffset.bWasApplied)
+	{
+		FNodalPointOffset NodalPointOffset;
+
+		if (LensFile->EvaluateNodalPointOffset(OutFocus, OutZoom, NodalPointOffset))
+		{
+			CachedNodalOffset.SetTranslation(NodalPointOffset.LocationOffset);
+			CachedNodalOffset.SetRotation(NodalPointOffset.RotationOffset);
+		}
+	}
+
+	// Compute an initial nodal offset transform with the Entrance Pupil Position equal to zero
+	EntrancePupilPosition = 0.0f;
+	OutNodalOffset = ComputeNodalOffset();
+
+	LastSavedNodalOffset = OutNodalOffset;
+	LastSavedEntrancePupilPosition = EntrancePupilPosition;
+
+	// Validate that the captured points all project to approximately the principal point,
+	// which would indicate that they were in roughly a straight line. 
+	const FLensDistortionState DistortionState = DistortionHandler->GetCurrentDistortionState();
+
+	cv::Mat CameraMatrix(3, 3, cv::DataType<double>::type);
+	cv::setIdentity(CameraMatrix);
+
+	CameraMatrix.at<double>(0, 0) = DistortionState.FocalLengthInfo.FxFy.X;
+	CameraMatrix.at<double>(1, 1) = DistortionState.FocalLengthInfo.FxFy.Y;
+
+	CameraMatrix.at<double>(0, 2) = DistortionState.ImageCenter.PrincipalPoint.X;
+	CameraMatrix.at<double>(1, 2) = DistortionState.ImageCenter.PrincipalPoint.Y;
+
+	// Compute the optimal camera pose using the new nodal offset and convert it to opencv's coordinate system
+	FTransform NewNodalOffsetTransform;
+	NewNodalOffsetTransform.SetLocation(OutNodalOffset.LocationOffset);
+	NewNodalOffsetTransform.SetRotation(OutNodalOffset.RotationOffset);
+
+	FTransform OptimalCameraPose = NewNodalOffsetTransform * CachedNodalOffset.Inverse() * CameraPose;
+
+	FOpenCVHelper::ConvertUnrealToOpenCV(OptimalCameraPose);
+
+	// Transform the 3D points to opencv's coordinate system
+	for (int32 Index = 0; Index < Points3d.size(); ++Index)
+	{
+		FTransform PointTransform = FTransform::Identity;
+		PointTransform.SetLocation(FVector(Points3d[Index].x, Points3d[Index].y, Points3d[Index].z));
+		FOpenCVHelper::ConvertUnrealToOpenCV(PointTransform);
+		Points3d[Index] = cv::Point3f(PointTransform.GetLocation().X, PointTransform.GetLocation().Y, PointTransform.GetLocation().Z);
+	}
+
+	// Initialize all of the corresponding 2D points to be exactly at the principal point of the lens
+	std::vector<cv::Point2f> Points2d;
+	for (int32 Index = 0; Index < Points3d.size(); ++Index)
+	{
+		Points2d.push_back(cv::Point2f(DistortionState.ImageCenter.PrincipalPoint.X, DistortionState.ImageCenter.PrincipalPoint.Y));
+	}
+
+	OutError = FOpenCVHelper::ComputeReprojectionError(OptimalCameraPose, CameraMatrix, Points3d, Points2d);
+
+	bFocusZoomWarningIssued = false;
+
+	return true;
+#else
+	{
+		OutErrorMessage = LOCTEXT("OpenCVRequired", "OpenCV is required");
+		return false;
+	}
+#endif // WITH_OPENCV
+}
+
+UMaterialInterface* UCameraNodalOffsetAlgoOpticalAxis::GetOverlayMaterial() const
+{
+	return TSoftObjectPtr<UMaterialInterface>(FSoftObjectPath(TEXT("/CameraCalibration/Materials/M_OutputTexture.M_OutputTexture"))).LoadSynchronous();
+}
+
+bool UCameraNodalOffsetAlgoOpticalAxis::IsOverlayEnabled() const
+{
+	// Always display the crosshair overlay while this algo is active
+	return true;
+}
+
+FNodalPointOffset UCameraNodalOffsetAlgoOpticalAxis::ComputeNodalOffset() const
+{
+	const FVector CameraLocation = PointAlongAxis + (EntrancePupilPosition * OpticalAxis.GetValue());
+
+	FTransform DesiredCameraTransform;
+	DesiredCameraTransform.SetLocation(CameraLocation);
+	DesiredCameraTransform.SetRotation(LookAtRotation.Quaternion());
+
+	const FTransform DesiredOffset = DesiredCameraTransform * CachedInverseCameraPose * CachedNodalOffset;
+
+	FNodalPointOffset NodalOffset;
+	NodalOffset.LocationOffset = DesiredOffset.GetLocation();
+	NodalOffset.RotationOffset = DesiredOffset.GetRotation();
+
+	return NodalOffset;
+}
+
+bool UCameraNodalOffsetAlgoOpticalAxis::UpdateNodalOffset()
+{
+	// Do not update if the optical axis has not been found yet
+	if (!OpticalAxis.IsSet())
+	{
+		return false;
+	}
+
+	if (!NodalOffsetTool.IsValid())
+	{
+		return false;
+	}
+
+	const FCameraCalibrationStepsController* StepsController = NodalOffsetTool->GetCameraCalibrationStepsController();
+
+	if (!StepsController)
+	{
+		return false;
+	}
+
+	ULensFile* LensFile = StepsController->GetLensFile();
+
+	if (!LensFile)
+	{
+		return false;
+	}
+
+	AdjustedNodalOffset = ComputeNodalOffset();
+
+	// Attempt to modify an existing point in the nodal offset table
+	bool bPointExistsInTable = LensFile->NodalOffsetTable.SetPoint(CachedFocus, CachedZoom, AdjustedNodalOffset);
+
+	return bPointExistsInTable;
+}
+
+void UCameraNodalOffsetAlgoOpticalAxis::SaveUpdatedNodalOffset()
+{
+	if (!NodalOffsetTool.IsValid())
+	{
+		return;
+	}
+
+	const FCameraCalibrationStepsController* StepsController = NodalOffsetTool->GetCameraCalibrationStepsController();
+
+	if (!StepsController)
+	{
+		return;
+	}
+
+	ULensFile* LensFile = StepsController->GetLensFile();
+
+	if (!LensFile)
+	{
+		return;
+	}
+
+	// Before beginning a transaction, restore the nodal offset and entrance pupil position to their most recently saved states
+	const FNodalPointOffset CurrentNodalOffset = AdjustedNodalOffset;
+	const float CurrentEntrancePupilPosition = EntrancePupilPosition;
+
+	AdjustedNodalOffset = LastSavedNodalOffset;
+	EntrancePupilPosition = LastSavedEntrancePupilPosition;
+
+	LensFile->NodalOffsetTable.SetPoint(CachedFocus, CachedZoom, LastSavedNodalOffset);
+	{
+		FScopedTransaction Transaction(LOCTEXT("SaveAdjustedNodalOffset", "Save Adjusted Nodal Offset"));
+		this->Modify();
+		LensFile->Modify();
+
+		// After beginning the transaction, reapply the final modified nodal offset and entrance pupil position
+		AdjustedNodalOffset = CurrentNodalOffset;
+		EntrancePupilPosition = CurrentEntrancePupilPosition;
+		LensFile->NodalOffsetTable.SetPoint(CachedFocus, CachedZoom, AdjustedNodalOffset);
+
+		// Reset the last saved nodal offset and entrance pupil position to the adjusted values to indicate that the change has been written
+		LastSavedNodalOffset = AdjustedNodalOffset;
+		LastSavedEntrancePupilPosition = EntrancePupilPosition;
+	}
+}
+
+TSharedRef<SWidget> UCameraNodalOffsetAlgoOpticalAxis::BuildUI()
+{
+	return SNew(SVerticalBox)
+
+	+ SVerticalBox::Slot() // Calibrator picker
+		.VAlign(EVerticalAlignment::VAlign_Top)
+		.AutoHeight()
+		.MaxHeight(FCameraCalibrationWidgetHelpers::DefaultRowHeight)
+		[FCameraCalibrationWidgetHelpers::BuildLabelWidgetPair(LOCTEXT("Calibrator", "Calibrator"), BuildCalibrationDevicePickerWidget())]
+
+	+ SVerticalBox::Slot() // Calibrator component picker
+		.VAlign(EVerticalAlignment::VAlign_Top)
+		.AutoHeight()
+		.MaxHeight(FCameraCalibrationWidgetHelpers::DefaultRowHeight)
+		[FCameraCalibrationWidgetHelpers::BuildLabelWidgetPair(LOCTEXT("CalibratorComponents", "Calibrator Component(s)"), BuildCalibrationComponentPickerWidget())]
+
+	+ SVerticalBox::Slot() // Calibrator point names
+		.AutoHeight()
+		.MaxHeight(FCameraCalibrationWidgetHelpers::DefaultRowHeight)
+		[FCameraCalibrationWidgetHelpers::BuildLabelWidgetPair(LOCTEXT("CalibratorPoint", "Calibrator Point"), BuildCalibrationPointsComboBox())]
+
+	+ SVerticalBox::Slot() // Entrance Pupil Position
+		.AutoHeight()
+		.MaxHeight(FCameraCalibrationWidgetHelpers::DefaultRowHeight)
+		[FCameraCalibrationWidgetHelpers::BuildLabelWidgetPair(LOCTEXT("EntrancePupilPosition", "Entrance Pupil Position"), BuildEntrancePupilPositionWidget())]
+
+	+ SVerticalBox::Slot() // Entrance Pupil Position
+		.AutoHeight()
+		.MaxHeight(FCameraCalibrationWidgetHelpers::DefaultRowHeight)
+		[FCameraCalibrationWidgetHelpers::BuildLabelWidgetPair(LOCTEXT("Sensitivity", "Sensitivity"), BuildSensitivityWidget())]
+
+	+ SVerticalBox::Slot() // Calibration Rows
+		.AutoHeight()
+		.MaxHeight(12 * FCameraCalibrationWidgetHelpers::DefaultRowHeight)
+		[BuildCalibrationPointsTable()]
+
+	+ SVerticalBox::Slot() // Row manipulation
+		.AutoHeight()
+		[
+			SNew(SHorizontalBox)
+	
+			+ SHorizontalBox::Slot() // Button to clear all rows
+			.AutoWidth()
+			[
+				SNew(SButton)
+				.Text(LOCTEXT("ClearAll", "Clear All"))
+				.HAlign(HAlign_Center)
+				.VAlign(VAlign_Center)
+				.OnClicked_Lambda([&]() -> FReply
+				{
+					ClearCalibrationRows();
+					return FReply::Handled();
+				})
+			]
+		]
+	;
+}
+
+TSharedRef<SWidget> UCameraNodalOffsetAlgoOpticalAxis::BuildEntrancePupilPositionWidget()
+{
+	return SNew(SNumericEntryBox<float>)
+		.Value_Lambda([&]() { return EntrancePupilPosition; })
+		.ToolTipText(LOCTEXT("EntrancePupil", "Entrance Pupil"))
+		.IsEnabled_Lambda([&] { return OpticalAxis.IsSet(); })
+		.OnValueChanged_Lambda([&](double InValue)
+		{
+			EntrancePupilPosition = InValue;
+			UpdateNodalOffset();
+			SaveUpdatedNodalOffset();
+		});
+}
+
+TSharedRef<SWidget> UCameraNodalOffsetAlgoOpticalAxis::BuildSensitivityWidget()
+{
+	return SNew(SComboBox<TSharedPtr<float>>)
+		.OptionsSource(&SensitivityOptions)
+		.IsEnabled_Lambda([&] { return OpticalAxis.IsSet(); })
+		.OnSelectionChanged_Lambda([&](TSharedPtr<float> NewValue, ESelectInfo::Type Type) -> void
+		{
+			AdjustmentIncrement = *NewValue;
+		})
+		.OnGenerateWidget_Lambda([&](TSharedPtr<float> InOption) -> TSharedRef<SWidget>
+		{
+			return SNew(STextBlock).Text(FText::AsNumber(*InOption));
+		})
+		.InitiallySelectedItem(SensitivityOptions[2])
+		[
+			SNew(STextBlock)
+			.Text_Lambda([&]() -> FText
+			{
+				return FText::AsNumber(AdjustmentIncrement);
+			})
+		];
+}
+
+TSharedRef<SWidget> UCameraNodalOffsetAlgoOpticalAxis::BuildHelpWidget()
+{
+	return SNew(STextBlock)
+		.Text(LOCTEXT("NodalOffsetAlgoOpticalAxisHelp",
+			"This nodal offset algorithm will find the nodal point of the lens and compute an offset based\n"
+			"on the current camera pose. First, it will use a set of 3D points that project to the center\n"
+			"of the image to find the optical axis of the lens. Then it will find the nodal point by\n"
+			"finding the entrance pupil point which lies along the optical axis.\n\n"
+
+			"The 3D points are taken from the calibrator object, which you need to select using the\n"
+			"provided picker. All that is required is that the object contains one or more 'Calibration\n'"
+			"Point Components'. These 3d calibration points will appear in the provided drop-down.\n"
+			"The calibrator you use should have some identifiable marker that will be easy to locate\n"
+			"in the viewport (such as a bright LED).\n\n"
+
+			"To build the set of 3d points that lie along the optical axis, align the marker on your\n"
+			"calibrator to the center of the crosshair that appears over the viewport. Note that the\n"
+			"crosshair may not be centered at the precise center of the image, but rather at the principal\n"
+			"point (which should already have been calibrated in the lens distortion calibration). You will\n"
+			"need to capture at least two points to find the optical axis, but more points may give more\n"
+			"accurate results.\n\n"
+
+			"Once those points are captured, the algorithm will find a best-fit line that passes through\n"
+			"the set of points. This is the optical axis of the lens, and the nodal point lies somewhere\n"
+			"along that axis.\n\n"
+
+			"The final step is to adjust the entrance pupil position until a tracked object in the scene\n"
+			"is visibly aligned with its CG counterpart in the simulcam viewport.\n\n"
+
+			"Notes:\n\n"
+			" - This calibration step relies on the camera having a lens distortion calibration.\n"
+			" - It requires the camera to not move much from the moment you capture the first\n"
+			"   point until you capture the last one.\n"
+		));
+}
+
+#undef LOCTEXT_NAMESPACE
diff --git a/Engine/Plugins/VirtualProduction/CameraCalibration/Source/CameraCalibrationEditor/Private/Calibrators/CameraNodalOffsetAlgoOpticalAxis.h b/Engine/Plugins/VirtualProduction/CameraCalibration/Source/CameraCalibrationEditor/Private/Calibrators/CameraNodalOffsetAlgoOpticalAxis.h
new file mode 100644
--- /dev/null
+++ b/Engine/Plugins/VirtualProduction/CameraCalibration/Source/CameraCalibrationEditor/Private/Calibrators/CameraNodalOffsetAlgoOpticalAxis.h
@@ -0,0 +1,106 @@
+// Copyright Epic Games, Inc. All Rights Reserved.
+
+#pragma once
+
+#include "CameraNodalOffsetAlgoPoints.h"
+
+#include "CameraNodalOffsetAlgoOpticalAxis.generated.h"
+
+/**
+ * Implements a nodal offset calibration algorithm that relies on the idea that the nodal point of the lens
+ * lies along the optical axis of the lens, and that the optical axis can be found by finding multiple points
+ * in 3D that map to the principal point of the lens in 2D.
+ *
+ * The algorithm uses 3d points (UCalibrationPointComponent) specified in a selectable calibrator with features 
+ * that the user can identify. By aligning that feature such that it projects to the principal point of the lens
+ * (easily identifiable using the automatically enabled crosshair overlay), the user captures two or more 3D
+ * points. The algorithm computes a line that connects those points, pointing away from the camera, which will
+ * be the optical axis. The nodal point is then found by moving the camera along that optical axis until the 
+ * calibrator appears perfectly aligned in the composite image (CG + Media).
+ */
+UCLASS()
+class UCameraNodalOffsetAlgoOpticalAxis : public UCameraNodalOffsetAlgoPoints
+{
+	GENERATED_BODY()
+
+public:
+	//~ Begin CalibPointsNodalOffsetAlgo
+	virtual void Initialize(UNodalOffsetTool* InNodalOffsetTool) override;
+	virtual void Tick(float DeltaTime) override;
+	virtual bool OnViewportInputKey(const FKey& InKey, const EInputEvent& InEvent) override;
+	virtual TSharedRef<SWidget> BuildUI() override;
+	virtual bool GetNodalOffset(FNodalPointOffset& OutNodalOffset, float& OutFocus, float& OutZoom, float& OutError, FText& OutErrorMessage) override;
+	virtual FName FriendlyName() const override { return TEXT("Nodal Offset Optical Axis"); };
+	virtual UMaterialInterface* GetOverlayMaterial() const override;
+	virtual bool IsOverlayEnabled() const override;
+	virtual TSharedRef<SWidget> BuildHelpWidget() override;
+	//~ End CalibPointsNodalOffsetAlgo
+
+protected:
+	/** The axis passing through the center of the lens, on which lies the nodal point */
+	TOptional<FVector> OpticalAxis;
+
+	/** A point that lies along the optical axis (not necessarily the nodal point) */
+	FVector PointAlongAxis = FVector::ZeroVector;
+
+	/** The rotation in which the camera is looking (along the optical axis) */
+	FRotator LookAtRotation = FRotator::ZeroRotator;
+
+	/** The current value of the nodal offset that has been written to the lens file, but not necessarily transacted */
+	UPROPERTY(Transient)
+	FNodalPointOffset AdjustedNodalOffset;
+
+	/** The last nodal offset to have been saved and transacted */
+	UPROPERTY(Transient)
+	FNodalPointOffset LastSavedNodalOffset;
+
+	/** The parametric position along the optical axis where the nodal points sits */
+	UPROPERTY(Transient)
+	float EntrancePupilPosition = 0.0f;
+
+	/** The last entrance pupil position to have been saved and transacted */
+	UPROPERTY(Transient)
+	float LastSavedEntrancePupilPosition = 0.0f;
+
+	/** Cached inverse of the camera pose when the points were captured */
+	FTransform CachedInverseCameraPose = FTransform::Identity;
+
+	/** Cached nodal offset of the camera when the points were captured */
+	FTransform CachedNodalOffset = FTransform::Identity;
+
+	/** Cached evaluated focus and zoom values when the points were captured */
+	float CachedFocus = 0.0f;
+	float CachedZoom = 0.0f;
+
+	/** Texture used to draw the crosshair overlay */
+	UPROPERTY(Transient)
+	TObjectPtr<UTextureRenderTarget2D> OverlayTexture;
+
+	/** Increment (in pixels) that the entrance pupil position will be adjusted by each individual key press */
+	float AdjustmentIncrement = 0.05f;
+
+	/** Array of sensitivity options to display to the user to set the adjustment increment */
+	TArray<TSharedPtr<float>> SensitivityOptions;
+
+	/** Used to avoid spamming the user with warnings about changing focus/zoom values */
+	bool bFocusZoomWarningIssued = false;
+
+	/** Used to track changes to the principal point so that the crosshair overlay can be updated */
+	FVector2D LastPrincipalPoint = FVector2D::ZeroVector;
+
+protected:
+	/** Build a widget that allows the user to manually enter a value for the entrance pupil position */
+	TSharedRef<SWidget> BuildEntrancePupilPositionWidget();
+
+	/** Builds the UI of the sensitivity widget */
+	TSharedRef<SWidget> BuildSensitivityWidget();
+
+	/** Compute the nodal offset based on the cached data and the current entrance pupil position */
+	FNodalPointOffset ComputeNodalOffset() const;
+
+	/** Recomputes the nodal offset and updates the lens file table */
+	bool UpdateNodalOffset();
+
+	/** Modifies the Lens File so that there is a transaction associated with the most recent nodal offset update */
+	void SaveUpdatedNodalOffset();
+};
diff --git a/Engine/Plugins/VirtualProduction/CameraCalibrationCore/Source/CameraCalibrationCore/Public/CameraNodalOffsetAlgo.h b/Engine/Plugins/VirtualProduction/CameraCalibrationCore/Source/CameraCalibrationCore/Public/CameraNodalOffsetAlgo.h
--- a/Engine/Plugins/VirtualProduction/CameraCalibrationCore/Source/CameraCalibrationCore/Public/CameraNodalOffsetAlgo.h
+++ b/Engine/Plugins/VirtualProduction/CameraCalibrationCore/Source/CameraCalibrationCore/Public/CameraNodalOffsetAlgo.h
@@ -3,12 +3,14 @@
 #pragma once
 
 #include "CoreMinimal.h"
+#include "Engine/EngineBaseTypes.h"
 #include "Widgets/SNullWidget.h"
 #include "Widgets/Text/STextBlock.h"
 
 #include "CameraNodalOffsetAlgo.generated.h"
 
 struct FGeometry;
+struct FKey;
 struct FNodalPointOffset;
 struct FPointerEvent;
 
@@ -39,6 +41,9 @@ public:
 	/** Callback when viewport is clicked. Returns false if the event was not handled. */
 	virtual bool OnViewportClicked(const FGeometry& MyGeometry, const FPointerEvent& MouseEvent) { return false;  };
 
+	/** Callback when viewport receives input key presses. Returns false if the event was not handled. */
+	virtual bool OnViewportInputKey(const FKey& InKey, const EInputEvent& InEvent) { return false; };
+
 	/** Returns the UI of this calibrator. Expected to only be called once */
 	virtual TSharedRef<SWidget> BuildUI() { return SNullWidget::NullWidget; };
 
