Remove "Save Lens Information" button from LensInfo step in LensFile editor. Changes made to the LensInfo in the details panel are now applied directly to the LensFile.

diff --git a/Engine/Plugins/VirtualProduction/CameraCalibration/Source/CameraCalibrationEditor/Private/AssetEditor/LensInfoStep.cpp b/Engine/Plugins/VirtualProduction/CameraCalibration/Source/CameraCalibrationEditor/Private/AssetEditor/LensInfoStep.cpp
--- a/Engine/Plugins/VirtualProduction/CameraCalibration/Source/CameraCalibrationEditor/Private/AssetEditor/LensInfoStep.cpp
+++ b/Engine/Plugins/VirtualProduction/CameraCalibration/Source/CameraCalibrationEditor/Private/AssetEditor/LensInfoStep.cpp
@@ -18,9 +18,9 @@
 void ULensInfoStep::Initialize(TWeakPtr<FCameraCalibrationStepsController> InCameraCalibrationStepController)
 {
 	CameraCalibrationStepsController = InCameraCalibrationStepController;
-	
-	EditedLensInfo = MakeShared<TStructOnScope<FLensInfo>>();
-	EditedLensInfo->InitializeAs<FLensInfo>();
+
+	const ULensFile* LensFile = CameraCalibrationStepsController.Pin()->GetLensFile();
+	OriginalLensInfo = LensFile->LensInfo;
 }
 
 TSharedRef<SWidget> ULensInfoStep::BuildUI()
@@ -31,7 +31,20 @@ TSharedRef<SWidget> ULensInfoStep::BuildUI()
 	DetailArgs.bShowScrollBar = true;
 
 	FPropertyEditorModule& PropertyEditor = FModuleManager::Get().LoadModuleChecked<FPropertyEditorModule>(TEXT("PropertyEditor"));
-	TSharedPtr<IStructureDetailsView> LensInfoStructureDetailsView = PropertyEditor.CreateStructureDetailView(DetailArgs, LensInfoStructDetailsView, EditedLensInfo);
+
+	ULensFile* LensFile = CameraCalibrationStepsController.Pin()->GetLensFile();
+	TSharedRef<FStructOnScope> StructOnScope = MakeShared<FStructOnScope>(FLensInfo::StaticStruct(), reinterpret_cast<uint8*>(&LensFile->LensInfo));
+
+	TSharedPtr<IStructureDetailsView> LensInfoStructureDetailsView = PropertyEditor.CreateStructureDetailView(DetailArgs, LensInfoStructDetailsView, StructOnScope);
+
+	LensInfoStructureDetailsView->GetOnFinishedChangingPropertiesDelegate().AddLambda([&](const FPropertyChangedEvent& PropertyChangedEvent)
+	{
+		// Ignore temporary interaction (dragging sliders, etc.)
+		if (PropertyChangedEvent.ChangeType == EPropertyChangeType::ValueSet)
+		{
+			OnSaveLensInformation();
+		}
+	});
 
 	TSharedPtr<SWidget> StepWidget = 
 		SNew(SVerticalBox)
@@ -51,19 +64,6 @@ TSharedRef<SWidget> ULensInfoStep::BuildUI()
 				.OnResetToDefault(FSimpleDelegate::CreateUObject(this, &ULensInfoStep::ResetToDefault))
 				.DiffersFromDefault(TAttribute<bool>::Create(TAttribute<bool>::FGetter::CreateUObject(this, &ULensInfoStep::DiffersFromDefault)))
 			]
-		]
-		+ SVerticalBox::Slot() // Save 
-		.AutoHeight()
-		.Padding(0, 20)
-		[
-			SNew(SButton).Text(LOCTEXT("SaveToLensFile", "Save Lens Information"))
-			.HAlign(HAlign_Center)
-			.VAlign(VAlign_Center)
-			.OnClicked_Lambda([this]() -> FReply
-			{
-				OnSaveLensInformation();
-				return FReply::Handled();
-			})
 		];
 	
 
@@ -86,8 +86,8 @@ void ULensInfoStep::Activate()
 	ULensFile* LensFile = CameraCalibrationStepsController.Pin()->GetLensFile();
 	if(LensFile)
 	{
-		//Initialize lens info to current value of the LensFile
-		*EditedLensInfo->Get() = LensFile->LensInfo;
+		// Initialize lens info to current value of the LensFile
+		CachedLensInfo = LensFile->LensInfo;
 	}
 
 	bIsActive = true;
@@ -98,7 +98,7 @@ void ULensInfoStep::Deactivate()
 	bIsActive = false;
 }
 
-void ULensInfoStep::OnSaveLensInformation() const
+void ULensInfoStep::OnSaveLensInformation()
 {
 	if (!CameraCalibrationStepsController.IsValid())
 	{
@@ -106,52 +106,56 @@ void ULensInfoStep::OnSaveLensInformation() const
 	}
 
 	ULensFile* LensFile = CameraCalibrationStepsController.Pin()->GetLensFile();
-	const FLensInfo& EditedData = *EditedLensInfo->Get();
-	if(EditedData == LensFile->LensInfo)
-	{
-		return;
-	}
 
-	//Validate sensor dimensions
-	constexpr float MinimumSize = 1.0f; //Limit sensor dimension to 1mm
-	if(EditedData.SensorDimensions.X < MinimumSize || EditedData.SensorDimensions.Y < MinimumSize)
+	// Validate sensor dimensions
+	constexpr float MinimumSize = 1.0f; // Limit sensor dimension to 1mm
+	if(LensFile->LensInfo.SensorDimensions.X < MinimumSize || LensFile->LensInfo.SensorDimensions.Y < MinimumSize)
 	{
 		const FText ErrorMessage = LOCTEXT("InvalidSensorDimensions", "Invalid sensor dimensions. Can't have dimensions smaller than 1mm");
 		FMessageDialog::Open(EAppMsgType::Ok, ErrorMessage);
+		
+		// Reset sensor dimensions if the new value was invalid
+		LensFile->LensInfo.SensorDimensions = CachedLensInfo.SensorDimensions;
 		return;	
 	}
 
-	//Initiate transaction in case we go through the path clearing tables
-	FScopedTransaction Transaction(LOCTEXT("ModifyingLensInfo", "Modifying LensFile information"));
-	LensFile->Modify();
-
-	//If lens model has changed and distortion table has data, notify the user he will lose calibration data if he continues
-	const bool bHasModelChanged = EditedData.LensModel != LensFile->LensInfo.LensModel;
-	const bool bHasSensorChanged = EditedData.SensorDimensions != LensFile->LensInfo.SensorDimensions;
-	if(bHasModelChanged || bHasSensorChanged)
+	// Before beginning the transaction, briefly restore the LensInfo to its previously saved state
+	const FLensInfo NewLensInfo = LensFile->LensInfo;
+	LensFile->LensInfo = CachedLensInfo;
 	{
-		if(LensFile->HasSamples(ELensDataCategory::Distortion)
-			|| LensFile->HasSamples(ELensDataCategory::ImageCenter)
-			|| LensFile->HasSamples(ELensDataCategory::Zoom))
+		FScopedTransaction Transaction(LOCTEXT("ModifyingLensInfo", "Modifying LensFile information"));
+		this->Modify();
+		LensFile->Modify();
+
+		// After beginning the transaction, reapply the modified LensInfo to the LensFile
+		LensFile->LensInfo = NewLensInfo;
+
+		// If lens model has changed and distortion table has data, notify the user that calibration data will be lost
+		const bool bHasModelChanged = CachedLensInfo.LensModel != LensFile->LensInfo.LensModel;
+		if (bHasModelChanged)
 		{
-			const FText DataChangeText = (bHasModelChanged && bHasSensorChanged) ? LOCTEXT("ModelAndSensorChange", "Lens model and sensor dimensions") : bHasModelChanged ? LOCTEXT("LensModelChange", "Lens model") : LOCTEXT("SensorDimensionsChange", "Sensor dimensions");
-			const FText Message = FText::Format(LOCTEXT("DataChangeDataLoss", "{0} change detected. Distortion, ImageCenter and FocalLength data samples will be lost. Do you wish to continue?")
-										, DataChangeText);
-			if (FMessageDialog::Open(EAppMsgType::OkCancel, Message) != EAppReturnType::Ok)
+			if (LensFile->HasSamples(ELensDataCategory::Distortion)
+				|| LensFile->HasSamples(ELensDataCategory::ImageCenter)
+				|| LensFile->HasSamples(ELensDataCategory::Zoom))
 			{
-				Transaction.Cancel();
-				return;
+				const FText Message = LOCTEXT("DataChangeDataLoss", "Lens model change detected. Distortion, ImageCenter and FocalLength data samples will be lost. Do you wish to continue?");
+				if (FMessageDialog::Open(EAppMsgType::OkCancel, Message) != EAppReturnType::Ok)
+				{
+					Transaction.Cancel();
+					LensFile->LensInfo.LensModel = CachedLensInfo.LensModel;
+					return;
+				}
 			}
+
+			// Clear table when disruptive change is detected and user proceeds
+			LensFile->ClearData(ELensDataCategory::Distortion);
+			LensFile->ClearData(ELensDataCategory::ImageCenter);
+			LensFile->ClearData(ELensDataCategory::Zoom);
 		}
 
-		//Clear table when disruptive change is detected and user proceeds
-		LensFile->ClearData(ELensDataCategory::Distortion);
-		LensFile->ClearData(ELensDataCategory::ImageCenter);
-		LensFile->ClearData(ELensDataCategory::Zoom);
+		// Cache the newly modified LensInfo as the most recently saved state
+		CachedLensInfo = LensFile->LensInfo;
 	}
-
-	//Apply new data to LensFile
-	LensFile->LensInfo = EditedData;
 }
 
 bool ULensInfoStep::IsActive() const
@@ -159,16 +163,17 @@ bool ULensInfoStep::IsActive() const
 	return bIsActive;
 }
 
-void ULensInfoStep::ResetToDefault() const
+void ULensInfoStep::ResetToDefault()
 {
-	const ULensFile* LensFile = CameraCalibrationStepsController.Pin()->GetLensFile();
-	*EditedLensInfo->Get() = LensFile->LensInfo;
+	ULensFile* LensFile = CameraCalibrationStepsController.Pin()->GetLensFile();
+	LensFile->LensInfo = OriginalLensInfo;
+	OnSaveLensInformation();
 }
 
 bool ULensInfoStep::DiffersFromDefault() const
 {
 	const ULensFile* LensFile = CameraCalibrationStepsController.Pin()->GetLensFile();
-	return !(*EditedLensInfo->Get() == LensFile->LensInfo);
+	return !(OriginalLensInfo == LensFile->LensInfo);
 }
 
 
diff --git a/Engine/Plugins/VirtualProduction/CameraCalibration/Source/CameraCalibrationEditor/Private/AssetEditor/LensInfoStep.h b/Engine/Plugins/VirtualProduction/CameraCalibration/Source/CameraCalibrationEditor/Private/AssetEditor/LensInfoStep.h
--- a/Engine/Plugins/VirtualProduction/CameraCalibration/Source/CameraCalibrationEditor/Private/AssetEditor/LensInfoStep.h
+++ b/Engine/Plugins/VirtualProduction/CameraCalibration/Source/CameraCalibrationEditor/Private/AssetEditor/LensInfoStep.h
@@ -31,13 +31,13 @@ public:
 
 public:
 
-	/** Triggered when user click save data button */
-	void OnSaveLensInformation() const;
+	/** Triggered when user edits any of the LensInfo properties */
+	void OnSaveLensInformation();
 
-	/** Reapplied current value from LensFile in our temporary copy displayed in UI */
-	void ResetToDefault() const;
+	/** Reapply the original LensInfo to the LensFile */
+	void ResetToDefault();
 
-	/** Returns true if our LensInfo differs from the one in the LensFile being edited */
+	/** Returns true if our LensInfo differs from the original info when the tool was opened */
 	bool DiffersFromDefault() const;
 
 private:
@@ -48,6 +48,10 @@ private:
 	/** True if this tool is the active one in the panel */
 	bool bIsActive = false;
 
-	/** Lens information being edited in the panel */
-	TSharedPtr<TStructOnScope<FLensInfo>> EditedLensInfo;
+	/** Original LensInfo, used to reset the details back to default */
+	FLensInfo OriginalLensInfo;
+
+	/** Cached previous LensInfo, used to trigger correct transactions to allow for undo operations */
+	UPROPERTY(Transient)
+	FLensInfo CachedLensInfo;
 };
