[CameraCalibration] Fix LensInfo step to be Transactional to properly support undo

diff --git a/Engine/Plugins/VirtualProduction/CameraCalibration/Source/CameraCalibrationEditor/Private/AssetEditor/CameraCalibrationStepsController.cpp b/Engine/Plugins/VirtualProduction/CameraCalibration/Source/CameraCalibrationEditor/Private/AssetEditor/CameraCalibrationStepsController.cpp
--- a/Engine/Plugins/VirtualProduction/CameraCalibration/Source/CameraCalibrationEditor/Private/AssetEditor/CameraCalibrationStepsController.cpp
+++ b/Engine/Plugins/VirtualProduction/CameraCalibration/Source/CameraCalibrationEditor/Private/AssetEditor/CameraCalibrationStepsController.cpp
@@ -152,7 +152,8 @@ void FCameraCalibrationStepsController::CreateSteps()
 		UCameraCalibrationStep* const Step = NewObject<UCameraCalibrationStep>(
 			GetTransientPackage(),
 			StepClass,
-			MakeUniqueObjectName(GetTransientPackage(), StepClass));
+			MakeUniqueObjectName(GetTransientPackage(), StepClass),
+			RF_Transactional);
 
 		check(Step);
 
