Add scroll boxes to the camera calibration right and bottom panels to avoid widgets getting squished.

diff --git a/Engine/Plugins/VirtualProduction/CameraCalibration/Source/CameraCalibrationEditor/Private/AssetEditor/SCameraCalibrationSteps.cpp b/Engine/Plugins/VirtualProduction/CameraCalibration/Source/CameraCalibrationEditor/Private/AssetEditor/SCameraCalibrationSteps.cpp
--- a/Engine/Plugins/VirtualProduction/CameraCalibration/Source/CameraCalibrationEditor/Private/AssetEditor/SCameraCalibrationSteps.cpp
+++ b/Engine/Plugins/VirtualProduction/CameraCalibration/Source/CameraCalibrationEditor/Private/AssetEditor/SCameraCalibrationSteps.cpp
@@ -27,6 +27,7 @@
 #include "Widgets/Layout/SSpacer.h"
 #include "Widgets/Layout/SWidgetSwitcher.h"
 #include "Widgets/Layout/SBorder.h"
+#include "Widgets/Layout/SScrollBox.h"
 #include "Widgets/SBoxPanel.h"
 #include "Widgets/SOverlay.h"
 #include "Widgets/Text/STextBlock.h"
@@ -92,102 +93,106 @@ void SCameraCalibrationSteps::Construct(const FArguments& InArgs, TWeakPtr<FCame
 		+ SHorizontalBox::Slot() // Right toolbar
 		.FillWidth(0.25f)
 		[ 
-			SNew(SVerticalBox)
-
-			+ SVerticalBox::Slot() // Viewport Title
-			.Padding(0, 5)
-			.AutoHeight()
-			.VAlign(EVerticalAlignment::VAlign_Center)
+			SNew(SScrollBox)
+			+ SScrollBox::Slot()
 			[
-				SNew(SBox)
-				.MinDesiredHeight(FCameraCalibrationWidgetHelpers::DefaultRowHeight)
-				.MaxDesiredHeight(FCameraCalibrationWidgetHelpers::DefaultRowHeight)
+				SNew(SVerticalBox)
+
+				+ SVerticalBox::Slot() // Viewport Title
+				.Padding(0, 5)
+				.AutoHeight()
+				.VAlign(EVerticalAlignment::VAlign_Center)
 				[
-					SNew(SBorder) // Background color for title
-					.BorderImage(FEditorStyle::GetBrush("DetailsView.CategoryTop"))
-					.BorderBackgroundColor(FLinearColor::White)
-					.VAlign(EVerticalAlignment::VAlign_Center)
+					SNew(SBox)
+					.MinDesiredHeight(FCameraCalibrationWidgetHelpers::DefaultRowHeight)
+					.MaxDesiredHeight(FCameraCalibrationWidgetHelpers::DefaultRowHeight)
 					[
-						SNew(SOverlay) 
-						+ SOverlay::Slot() // Used to add left padding to the title
-						.Padding(5,0,0,0)
+						SNew(SBorder) // Background color for title
+						.BorderImage(FEditorStyle::GetBrush("DetailsView.CategoryTop"))
+						.BorderBackgroundColor(FLinearColor::White)
+						.VAlign(EVerticalAlignment::VAlign_Center)
 						[
-							SNew(STextBlock) // Title text
-							.Text(LOCTEXT("ViewportSettings", "Viewport Settings"))
-							.TransformPolicy(ETextTransformPolicy::ToUpper)
-							.Font(FAppStyle::Get().GetFontStyle(TEXT("PropertyWindow.BoldFont")))
-							.TextStyle(FAppStyle::Get(), "DetailsView.CategoryTextStyle")
+							SNew(SOverlay) 
+							+ SOverlay::Slot() // Used to add left padding to the title
+							.Padding(5,0,0,0)
+							[
+								SNew(STextBlock) // Title text
+								.Text(LOCTEXT("ViewportSettings", "Viewport Settings"))
+								.TransformPolicy(ETextTransformPolicy::ToUpper)
+								.Font(FAppStyle::Get().GetFontStyle(TEXT("PropertyWindow.BoldFont")))
+								.TextStyle(FAppStyle::Get(), "DetailsView.CategoryTextStyle")
+							]
 						]
 					]
 				]
-			]
 
-			+ SVerticalBox::Slot() // Wiper
-			.MaxHeight(FCameraCalibrationWidgetHelpers::DefaultRowHeight)
-			[ FCameraCalibrationWidgetHelpers::BuildLabelWidgetPair(LOCTEXT("Transparency", "Transparency"), BuildSimulcamWiperWidget())]
+				+ SVerticalBox::Slot() // Wiper
+				.MaxHeight(FCameraCalibrationWidgetHelpers::DefaultRowHeight)
+				[ FCameraCalibrationWidgetHelpers::BuildLabelWidgetPair(LOCTEXT("Transparency", "Transparency"), BuildSimulcamWiperWidget())]
 				
-			+ SVerticalBox::Slot() // Camera picker
-			.MaxHeight(FCameraCalibrationWidgetHelpers::DefaultRowHeight)
-			[ FCameraCalibrationWidgetHelpers::BuildLabelWidgetPair(LOCTEXT("Camera", "Camera"), BuildCameraPickerWidget())]
+				+ SVerticalBox::Slot() // Camera picker
+				.MaxHeight(FCameraCalibrationWidgetHelpers::DefaultRowHeight)
+				[ FCameraCalibrationWidgetHelpers::BuildLabelWidgetPair(LOCTEXT("Camera", "Camera"), BuildCameraPickerWidget())]
 				
-			+ SVerticalBox::Slot() // Media Source picker
-			.MaxHeight(FCameraCalibrationWidgetHelpers::DefaultRowHeight)
-			[FCameraCalibrationWidgetHelpers::BuildLabelWidgetPair(LOCTEXT("MediaSource", "Media Source"), BuildMediaSourceWidget())]
+				+ SVerticalBox::Slot() // Media Source picker
+				.MaxHeight(FCameraCalibrationWidgetHelpers::DefaultRowHeight)
+				[FCameraCalibrationWidgetHelpers::BuildLabelWidgetPair(LOCTEXT("MediaSource", "Media Source"), BuildMediaSourceWidget())]
 
-			+ SVerticalBox::Slot() // Overlay picker
-			.MaxHeight(FCameraCalibrationWidgetHelpers::DefaultRowHeight)
-			[FCameraCalibrationWidgetHelpers::BuildLabelWidgetPair(LOCTEXT("Overlay", "Overlay"), BuildOverlayWidget())]
+				+ SVerticalBox::Slot() // Overlay picker
+				.MaxHeight(FCameraCalibrationWidgetHelpers::DefaultRowHeight)
+				[FCameraCalibrationWidgetHelpers::BuildLabelWidgetPair(LOCTEXT("Overlay", "Overlay"), BuildOverlayWidget())]
 
-			+ SVerticalBox::Slot() // Overlay parameters
-			.AutoHeight()
-			[
-				OverlayParameterWidget.ToSharedRef()
-			]
+				+ SVerticalBox::Slot() // Overlay parameters
+				.AutoHeight()
+				[
+					OverlayParameterWidget.ToSharedRef()
+				]
 
-			+ SVerticalBox::Slot() // Step Title
-			.Padding(0, 5)
-			.AutoHeight()
-			.VAlign(EVerticalAlignment::VAlign_Center)
-			[
-				SNew(SBox) // Constrain the height
-				.MinDesiredHeight(FCameraCalibrationWidgetHelpers::DefaultRowHeight)
-				.MaxDesiredHeight(FCameraCalibrationWidgetHelpers::DefaultRowHeight)
+				+ SVerticalBox::Slot() // Step Title
+				.Padding(0, 5)
+				.AutoHeight()
+				.VAlign(EVerticalAlignment::VAlign_Center)
 				[
-					SNew(SBorder) // Background color of title
-					.BorderImage(FEditorStyle::GetBrush("DetailsView.CategoryTop"))
-					.BorderBackgroundColor(FLinearColor::White)
-					.VAlign(EVerticalAlignment::VAlign_Center)
+					SNew(SBox) // Constrain the height
+					.MinDesiredHeight(FCameraCalibrationWidgetHelpers::DefaultRowHeight)
+					.MaxDesiredHeight(FCameraCalibrationWidgetHelpers::DefaultRowHeight)
 					[
-						SNew(SOverlay) 
-						+ SOverlay::Slot() // Used to add left padding to the title
-						.Padding(5, 0, 0, 0)
+						SNew(SBorder) // Background color of title
+						.BorderImage(FEditorStyle::GetBrush("DetailsView.CategoryTop"))
+						.BorderBackgroundColor(FLinearColor::White)
+						.VAlign(EVerticalAlignment::VAlign_Center)
 						[
-							SNew(STextBlock) // Title text
-							.Text_Lambda([&]() -> FText
-							{
-								for (const TStrongObjectPtr<UCameraCalibrationStep>& Step : CalibrationStepsController.Pin()->GetCalibrationSteps())
+							SNew(SOverlay) 
+							+ SOverlay::Slot() // Used to add left padding to the title
+							.Padding(5, 0, 0, 0)
+							[
+								SNew(STextBlock) // Title text
+								.Text_Lambda([&]() -> FText
 								{
-									if (!Step.IsValid() || !Step->IsActive())
+									for (const TStrongObjectPtr<UCameraCalibrationStep>& Step : CalibrationStepsController.Pin()->GetCalibrationSteps())
 									{
-										continue;
-									}
+										if (!Step.IsValid() || !Step->IsActive())
+										{
+											continue;
+										}
 
-									return FText::FromName(Step->FriendlyName());
-								}
+										return FText::FromName(Step->FriendlyName());
+									}
 
-								return LOCTEXT("StepSettings", "Step");
-							})
-							.TransformPolicy(ETextTransformPolicy::ToUpper)
-							.Font(FAppStyle::Get().GetFontStyle(TEXT("PropertyWindow.BoldFont")))
-							.TextStyle(FAppStyle::Get(), "DetailsView.CategoryTextStyle")
+									return LOCTEXT("StepSettings", "Step");
+								})
+								.TransformPolicy(ETextTransformPolicy::ToUpper)
+								.Font(FAppStyle::Get().GetFontStyle(TEXT("PropertyWindow.BoldFont")))
+								.TextStyle(FAppStyle::Get(), "DetailsView.CategoryTextStyle")
+							]
 						]
 					]
 				]
-			]
 
-			+ SVerticalBox::Slot() // Step UI
-			.AutoHeight()
-			[StepWidgetSwitcher.ToSharedRef()]
+				+ SVerticalBox::Slot() // Step UI
+				.AutoHeight()
+				[StepWidgetSwitcher.ToSharedRef()]
+			]
 		]
 	];
 
diff --git a/Engine/Plugins/VirtualProduction/CameraCalibration/Source/CameraCalibrationEditor/Private/AssetEditor/SLensEvaluation.cpp b/Engine/Plugins/VirtualProduction/CameraCalibration/Source/CameraCalibrationEditor/Private/AssetEditor/SLensEvaluation.cpp
--- a/Engine/Plugins/VirtualProduction/CameraCalibration/Source/CameraCalibrationEditor/Private/AssetEditor/SLensEvaluation.cpp
+++ b/Engine/Plugins/VirtualProduction/CameraCalibration/Source/CameraCalibrationEditor/Private/AssetEditor/SLensEvaluation.cpp
@@ -13,6 +13,7 @@
 #include "UI/CameraCalibrationEditorStyle.h"
 #include "Widgets/Input/SCheckBox.h"
 #include "Widgets/Layout/SGridPanel.h"
+#include "Widgets/Layout/SScrollBox.h"
 #include "Widgets/SBoxPanel.h"
 #include "Widgets/Text/STextBlock.h"
 
@@ -30,49 +31,53 @@ void SLensEvaluation::Construct(const FArguments& InArgs, TWeakPtr<FCameraCalibr
 		SNew(SBorder)
 		.Padding(FMargin(4.f, 4.f, 4.f, 4.f))
 		[
-			SNew(SHorizontalBox)
-
-			//Tracking section
-			+ SHorizontalBox::Slot()
-			.Padding(5.0f, 5.0f)
-			.FillWidth(0.2f)
-			[
-				MakeTrackingWidget()
-			]
-			//Raw Input FIZ section
-			+ SHorizontalBox::Slot()
-			.Padding(5.0f, 5.0f)
-			.FillWidth(0.2f)
-			[
-				MakeRawInputFIZWidget()
-			]
-			//Evaluated FIZ section
-			+ SHorizontalBox::Slot()
-			.Padding(5.0f, 5.0f)
-			.FillWidth(0.2f)
-			[
-				MakeEvaluatedFIZWidget()
-			]
-			//Distortion section
-			+ SHorizontalBox::Slot()
-			.Padding(5.0f, 5.0f)
-			.FillWidth(0.2f)
-			[
-				MakeDistortionWidget()
-			]
-			//Image Center section
-			+ SHorizontalBox::Slot()
-			.Padding(5.0f, 5.0f)
-			.FillWidth(0.2f)
-			[
-				MakeIntrinsicsWidget()
-			]
-			//Nodal offset section
-			+ SHorizontalBox::Slot()
-			.Padding(5.0f, 5.0f)
-			.FillWidth(0.2f)
-			[
-				MakeNodalOffsetWidget()	
+			SNew(SScrollBox)
+			+ SScrollBox::Slot()
+			[
+				SNew(SHorizontalBox)
+
+				//Tracking section
+				+ SHorizontalBox::Slot()
+				.Padding(5.0f, 5.0f)
+				.FillWidth(0.2f)
+				[
+					MakeTrackingWidget()
+				]
+				//Raw Input FIZ section
+				+ SHorizontalBox::Slot()
+				.Padding(5.0f, 5.0f)
+				.FillWidth(0.2f)
+				[
+					MakeRawInputFIZWidget()
+				]
+				//Evaluated FIZ section
+				+ SHorizontalBox::Slot()
+				.Padding(5.0f, 5.0f)
+				.FillWidth(0.2f)
+				[
+					MakeEvaluatedFIZWidget()
+				]
+				//Distortion section
+				+ SHorizontalBox::Slot()
+				.Padding(5.0f, 5.0f)
+				.FillWidth(0.2f)
+				[
+					MakeDistortionWidget()
+				]
+				//Image Center section
+				+ SHorizontalBox::Slot()
+				.Padding(5.0f, 5.0f)
+				.FillWidth(0.2f)
+				[
+					MakeIntrinsicsWidget()
+				]
+				//Nodal offset section
+				+ SHorizontalBox::Slot()
+				.Padding(5.0f, 5.0f)
+				.FillWidth(0.2f)
+				[
+					MakeNodalOffsetWidget()	
+				]
 			]
 		]
 	];
