[CameraCalibration] Distortion calibration algo based on an arbitrary set of 2D-3D point correspondences

diff --git a/Engine/Plugins/VirtualProduction/CameraCalibration/Source/CameraCalibrationEditor/Private/Calibrators/CameraLensDistortionAlgoPoints.cpp b/Engine/Plugins/VirtualProduction/CameraCalibration/Source/CameraCalibrationEditor/Private/Calibrators/CameraLensDistortionAlgoPoints.cpp
new file mode 100644
--- /dev/null
+++ b/Engine/Plugins/VirtualProduction/CameraCalibration/Source/CameraCalibrationEditor/Private/Calibrators/CameraLensDistortionAlgoPoints.cpp
@@ -0,0 +1,938 @@
+// Copyright Epic Games, Inc. All Rights Reserved.
+
+#include "CameraLensDistortionAlgoPoints.h"
+
+#include "AssetEditor/CameraCalibrationStepsController.h"
+#include "AssetEditor/LensDistortionTool.h"
+#include "CalibrationPointComponent.h"
+#include "Camera/CameraActor.h"
+#include "Editor.h"
+#include "Input/Events.h"
+#include "Math/Vector.h"
+#include "Misc/MessageDialog.h"
+#include "Models/SphericalLensModel.h"
+#include "UI/SFilterableActorPicker.h"
+#include "UI/CameraCalibrationWidgetHelpers.h"
+#include "Widgets/Input/SButton.h"
+#include "Widgets/Input/SNumericEntryBox.h"
+#include "Widgets/Views/SListView.h"
+
+#if WITH_OPENCV
+
+#include "OpenCVHelper.h"
+#include <vector>
+
+#include "PreOpenCVHeaders.h"
+#include "opencv2/calib3d.hpp"
+#include "PostOpenCVHeaders.h"
+
+#endif
+
+#define LOCTEXT_NAMESPACE "CameraLensDistortionAlgoPoints"
+
+namespace CameraLensDistortionAlgoPoints
+{
+	class SCalibrationRowGenerator : public SMultiColumnTableRow<TSharedPtr<UCameraLensDistortionAlgoPoints::FCalibrationRowData>>
+	{
+		using FCalibrationRowData = UCameraLensDistortionAlgoPoints::FCalibrationRowData;
+
+	public:
+		SLATE_BEGIN_ARGS(SCalibrationRowGenerator) {}
+
+		/** The list item for this row */
+		SLATE_ARGUMENT(TSharedPtr<FCalibrationRowData>, CalibrationRowData)
+
+		SLATE_END_ARGS()
+
+		void Construct(const FArguments& Args, const TSharedRef<STableViewBase>& OwnerTableView)
+		{
+			CalibrationRowData = Args._CalibrationRowData;
+
+			SMultiColumnTableRow<TSharedPtr<FCalibrationRowData>>::Construct(
+				FSuperRowType::FArguments()
+				.Padding(1.0f),
+				OwnerTableView
+			);
+		}
+
+		//~Begin SMultiColumnTableRow
+		virtual TSharedRef<SWidget> GenerateWidgetForColumn(const FName& ColumnName) override
+		{
+			if (ColumnName == TEXT("Pattern"))
+			{
+				const FString Text = FString::Printf(TEXT("%d"), CalibrationRowData->PatternIndex + 1);
+				return SNew(STextBlock).Text(FText::FromString(Text));
+			}
+
+			if (ColumnName == TEXT("Name"))
+			{
+				const FString Text = CalibrationRowData->CalibratorPointData.Name;
+				return SNew(STextBlock).Text(FText::FromString(Text));
+			}
+
+			if (ColumnName == TEXT("Point2D"))
+			{
+				const FString Text = FString::Printf(TEXT("(%.2f, %.2f)"),
+					CalibrationRowData->CalibratorPointData.Point2d.X,
+					CalibrationRowData->CalibratorPointData.Point2d.Y);
+
+				return SNew(STextBlock).Text(FText::FromString(Text));
+			}
+
+			if (ColumnName == TEXT("Point3D"))
+			{
+				const FString Text = FString::Printf(TEXT("(%.0f, %.0f, %.0f)"),
+					CalibrationRowData->CalibratorPointData.Point3d.X,
+					CalibrationRowData->CalibratorPointData.Point3d.Y,
+					CalibrationRowData->CalibratorPointData.Point3d.Z);
+
+				return SNew(STextBlock).Text(FText::FromString(Text));
+			}
+
+			return SNullWidget::NullWidget;
+		}
+		//~End SMultiColumnTableRow
+
+
+	private:
+		TSharedPtr<FCalibrationRowData> CalibrationRowData;
+	};
+};
+
+void UCameraLensDistortionAlgoPoints::Initialize(ULensDistortionTool* InLensDistortionTool)
+{
+	LensDistortionTool = InLensDistortionTool;
+
+	// Guess which calibrator to use by searching for actors with CalibrationPointComponents.
+	SetCalibrator(FindFirstCalibrator());
+}
+
+void UCameraLensDistortionAlgoPoints::Shutdown()
+{
+	LensDistortionTool.Reset();
+}
+
+void UCameraLensDistortionAlgoPoints::Tick(float DeltaTime)
+{
+	if (!LensDistortionTool.IsValid())
+	{
+		return;
+	}
+
+	const FCameraCalibrationStepsController* StepsController = LensDistortionTool->GetCameraCalibrationStepsController();
+
+	if (!StepsController)
+	{
+		return;
+	}
+
+	// If not paused, cache the 3d locations of each calibrator point associated with the current calibrator
+	if (!StepsController->IsPaused())
+	{
+		// Get calibration point data
+		do
+		{
+			LastCalibratorPoints.Empty();
+
+			for (const TSharedPtr<FCalibratorPointData>& CalibratorPoint : CalibratorPoints)
+			{
+				if (!CalibratorPoint.IsValid())
+				{
+					continue;
+				}
+
+				FCalibratorPointCache PointCache;
+
+				if (!GetCalibratorPointCacheFromName(CalibratorPoint->Name, PointCache))
+				{
+					continue;
+				}
+
+				LastCalibratorPoints.Emplace(MoveTemp(PointCache));
+			}
+		} while (0);
+
+		// Get camera / calibrator data
+		do
+		{
+			LastCameraData.bIsValid = false;
+
+			const ACameraActor* Camera = StepsController->GetCamera();
+
+			if (!Camera)
+			{
+				break;
+			}
+
+			LastCameraData.UniqueId = Camera->GetUniqueID();
+
+			const FLensFileEvalData* LensFileEvalData = StepsController->GetLensFileEvalData();
+
+			if (!LensFileEvalData)
+			{
+				break;
+			}
+
+			LastCameraData.LensFileEvalData = *LensFileEvalData;
+
+			if (!Calibrator.IsValid())
+			{
+				break;
+			}
+
+			LastCameraData.CalibratorUniqueId = Calibrator->GetUniqueID();
+
+			LastCameraData.bIsValid = true;
+
+		} while (0);
+	}
+}
+
+bool UCameraLensDistortionAlgoPoints::OnViewportClicked(const FGeometry& MyGeometry, const FPointerEvent& MouseEvent)
+{
+	// We only respond to left clicks
+	if (MouseEvent.GetEffectingButton() != EKeys::LeftMouseButton)
+	{
+		return false;
+	}
+
+	if (!LensDistortionTool.IsValid())
+	{
+		return true;
+	}
+
+	FCameraCalibrationStepsController* StepsController = LensDistortionTool->GetCameraCalibrationStepsController();
+
+	if (!StepsController)
+	{
+		return true;
+	}
+
+	// Get currently selected calibrator point
+	FCalibratorPointCache LastCalibratorPoint;
+	LastCalibratorPoint.bIsValid = false;
+	{
+		if (!CurrentCalibratorPoint.IsValid())
+		{
+			return true;
+		}
+
+		// Find its values in the cache
+		for (const FCalibratorPointCache& PointCache : LastCalibratorPoints)
+		{
+			if (PointCache.bIsValid && (PointCache.Name == CurrentCalibratorPoint->Name))
+			{
+				LastCalibratorPoint = PointCache;
+				break;
+			}
+		}
+	}
+
+	// Check that we have a valid calibrator 3dpoint or camera data
+	if (!LastCalibratorPoint.bIsValid || !LastCameraData.bIsValid)
+	{
+		return true;
+	}
+
+	// Get the mouse click 2d position
+	if (!StepsController->CalculateNormalizedMouseClickPosition(MyGeometry, MouseEvent, LastCalibratorPoint.Point2d))
+	{
+		return true;
+	}
+
+	TSharedPtr<FCalibrationRowData> Row = MakeShared<FCalibrationRowData>();
+	Row->CameraData = LastCameraData;
+	Row->CalibratorPointData = LastCalibratorPoint;
+	Row->PatternIndex = CurrentPatternIndex;
+
+	// Validate the new row, show a message if validation fails.
+ 	{
+ 		FText ErrorMessage;
+ 
+ 		if (!ValidateNewRow(Row, ErrorMessage))
+ 		{
+ 			const FText TitleError = LOCTEXT("NewRowError", "New Row Error");
+ 			FMessageDialog::Open(EAppMsgType::Ok, ErrorMessage, &TitleError);
+ 			return true;
+ 		}
+ 	}
+
+	// Add this data point
+	CalibrationRows.Add(Row);
+
+	// Notify the ListView of the new data
+	if (CalibrationListView.IsValid())
+	{
+		CalibrationListView->RequestListRefresh();
+		CalibrationListView->RequestNavigateToItem(Row);
+	}
+
+	// Auto-advance to the next calibration point (if it exists)
+	if (AdvanceCalibratorPoint())
+	{
+		// Play media if this was the last point in the object
+		StepsController->Play();
+
+		// If the calibrator points have looped around to the first point again, advance the pattern index
+		++CurrentPatternIndex;
+	}
+
+	return true;
+}
+
+TSharedRef<SWidget> UCameraLensDistortionAlgoPoints::BuildUI()
+{
+	return SNew(SVerticalBox)
+
+		+ SVerticalBox::Slot() // Calibrator picker
+		.VAlign(EVerticalAlignment::VAlign_Top)
+		.AutoHeight()
+		.MaxHeight(FCameraCalibrationWidgetHelpers::DefaultRowHeight)
+		[FCameraCalibrationWidgetHelpers::BuildLabelWidgetPair(LOCTEXT("Calibrator", "Calibrator"), BuildCalibrationDevicePickerWidget())]
+
+		+ SVerticalBox::Slot() // Calibrator point names
+		.AutoHeight()
+		.MaxHeight(FCameraCalibrationWidgetHelpers::DefaultRowHeight)
+		[FCameraCalibrationWidgetHelpers::BuildLabelWidgetPair(LOCTEXT("CurrentCalibratorPoint", "Current Calibrator Point"), BuildCurrentCalibratorPointLabel())]
+
+		+ SVerticalBox::Slot() // Current focal length
+		.AutoHeight()
+		.MaxHeight(FCameraCalibrationWidgetHelpers::DefaultRowHeight)
+		[FCameraCalibrationWidgetHelpers::BuildLabelWidgetPair(LOCTEXT("CurrentFocalLength", "Current Focal Length (mm)"), BuildCurrentFocalLengthWidget())]
+
+		+ SVerticalBox::Slot() // Spacer
+		[
+			SNew(SBox)
+			.MinDesiredHeight(0.25 * FCameraCalibrationWidgetHelpers::DefaultRowHeight)
+			.MaxDesiredHeight(0.25 * FCameraCalibrationWidgetHelpers::DefaultRowHeight)
+		]
+
+		+ SVerticalBox::Slot() // Calibration Rows
+		.AutoHeight()
+		.MaxHeight(12 * FCameraCalibrationWidgetHelpers::DefaultRowHeight)
+		[BuildCalibrationPointsTable()]
+
+		+ SVerticalBox::Slot() // Action buttons (e.g. Remove, Clear)
+		.HAlign(EHorizontalAlignment::HAlign_Center)
+		.AutoHeight()
+		.Padding(0, 20)
+		[BuildCalibrationActionButtons()];
+}
+
+bool UCameraLensDistortionAlgoPoints::GetLensDistortion(
+	float& OutFocus,
+	float& OutZoom,
+	FDistortionInfo& OutDistortionInfo,
+	FFocalLengthInfo& OutFocalLengthInfo,
+	FImageCenterInfo& OutImageCenterInfo,
+	TSubclassOf<ULensModel>& OutLensModel,
+	double& OutError,
+	FText& OutErrorMessage)
+{
+#if WITH_OPENCV
+	if (!LensDistortionTool.IsValid())
+	{
+		OutErrorMessage = LOCTEXT("InvalidTool", "Invalid Tool");
+		return false;
+	}
+
+	const FCameraCalibrationStepsController* StepsController = LensDistortionTool->GetCameraCalibrationStepsController();
+
+	if (!StepsController)
+	{
+		OutErrorMessage = LOCTEXT("InvalidStepsController", "Invalid StepsController");
+		return false;
+	}
+
+	const ULensFile* LensFile = StepsController->GetLensFile();
+
+	if (!LensFile)
+	{
+		OutErrorMessage = LOCTEXT("InvalidLensFile", "Invalid Lens File");
+		return false;
+	}
+
+	const FIntPoint ImageSize = StepsController->GetCompRenderTargetSize();
+
+	// Gather 3D points and 2D points from each group of rows
+	std::vector<std::vector<cv::Point3f>> Samples3d;
+	std::vector<std::vector<cv::Point2f>> Samples2d;
+
+	uint32 NextPatternIndex = 0;
+	for (int32 RowIndex = 0; RowIndex < CalibrationRows.Num(); ++RowIndex)
+	{
+		std::vector<cv::Point3f> Points3d;
+		std::vector<cv::Point2f> Points2d;
+
+		while (RowIndex < CalibrationRows.Num() && CalibrationRows[RowIndex]->PatternIndex == NextPatternIndex)
+		{
+			const TSharedPtr<FCalibrationRowData>& Row = CalibrationRows[RowIndex];
+
+			Points3d.push_back(cv::Point3f(
+				Row->CalibratorPointData.Point3d.X,
+				Row->CalibratorPointData.Point3d.Y,
+				Row->CalibratorPointData.Point3d.Z));
+
+			Points2d.push_back(cv::Point2f(
+				Row->CalibratorPointData.Point2d.X,
+				Row->CalibratorPointData.Point2d.Y));
+
+			++RowIndex;
+		}
+
+		if (Points3d.size() > 0)
+		{
+			Samples3d.push_back(Points3d);
+			Samples2d.push_back(Points2d);
+		}
+
+		--RowIndex;
+		++NextPatternIndex;
+	}
+
+	// Validate that there are at least 4 patterns
+	if (Samples3d.size() < 4)
+	{
+		OutErrorMessage = LOCTEXT("NotEnoughSamples", "At least 4 calibration patterns are required");
+		return false;
+	}
+
+	// Validate that each pattern has the same number of points
+	const int32 NumPointsInPattern = Samples3d[0].size();
+	for (int32 PatternIndex = 1; PatternIndex < Samples3d.size(); ++PatternIndex)
+	{
+		if (Samples3d[PatternIndex].size() != NumPointsInPattern)
+		{
+			OutErrorMessage = LOCTEXT("DifferentNumPointsInPattern", "Every calibration pattern must have the same number of points");
+			return false;
+		}
+	}
+
+	// Because the calibration pattern is not coplanar, OpenCV requires an initial intrinsics guess
+	if (FMath::IsNearlyEqual(CurrentFocalLengthMM, 0.0f) || CurrentFocalLengthMM < 0.0f)
+	{
+		OutErrorMessage = LOCTEXT("InvalidFocalLengthMsg", "The current focal length (in mm) must be set to a "
+			"valid value in order to seed the calibration algorithm with a best guess as to the camera intrinsics (Fx/Fy)");
+		return false;
+	}
+
+	if (FMath::IsNearlyEqual(CurrentFocalLengthMM, 0.0f) || CurrentFocalLengthMM < 0.0f)
+	{
+		OutErrorMessage = LOCTEXT("InvalidFocalLengthMsg", "The current focal length (in mm) must be set to a "
+			"valid value in order to seed the calibration algorithm with a best guess as to the camera intrinsics (Fx/Fy)");
+		return false;
+	}
+
+	// Validate sensor dimensions
+	constexpr float MinimumSize = 1.0f; // Limit sensor dimension to 1mm
+	if (LensFile->LensInfo.SensorDimensions.X < MinimumSize || LensFile->LensInfo.SensorDimensions.Y < MinimumSize)
+	{
+		OutErrorMessage = LOCTEXT("InvalidSensorDimensions", "Invalid sensor dimensions. Can't have dimensions smaller than 1mm");
+		return false;
+	}
+
+	const float Fx = CurrentFocalLengthMM / LensFile->LensInfo.SensorDimensions.X;
+	const float Fy = CurrentFocalLengthMM / LensFile->LensInfo.SensorDimensions.Y;
+
+	cv::Mat CameraMatrix = cv::Mat::eye(3, 3, CV_64F);
+
+	CameraMatrix.at<double>(0, 0) = Fx * ImageSize.X; // Fx
+	CameraMatrix.at<double>(1, 1) = Fy * ImageSize.Y; // Fy
+	CameraMatrix.at<double>(0, 2) = 0.5 * ImageSize.X; // Cx
+	CameraMatrix.at<double>(1, 2) = 0.5 * ImageSize.Y; // Cy
+
+	cv::Mat DistortionCoefficients;
+
+	std::vector<cv::Mat> Rvecs;
+	std::vector<cv::Mat> Tvecs;
+
+ 	std::vector<double> IntrinsicsDeviations;
+ 	std::vector<double> ExtrinsicsDeviations;
+ 	std::vector<double> ViewErrors;
+
+	OutError = cv::calibrateCamera(
+		Samples3d,
+		Samples2d,
+		cv::Size(ImageSize.X, ImageSize.Y),
+		CameraMatrix,
+		DistortionCoefficients,
+		Rvecs,
+		Tvecs,
+		IntrinsicsDeviations,
+		ExtrinsicsDeviations,
+		ViewErrors,
+		cv::CALIB_USE_INTRINSIC_GUESS
+	);
+
+	check(DistortionCoefficients.total() == 5);
+	check((CameraMatrix.rows == 3) && (CameraMatrix.cols == 3));
+
+	const TSharedPtr<FCalibrationRowData>& FirstRow = CalibrationRows[0];
+
+	// FZ inputs to LUT
+	OutFocus = FirstRow->CameraData.LensFileEvalData.Input.Focus;
+	OutZoom = FirstRow->CameraData.LensFileEvalData.Input.Zoom;
+
+	// FocalLengthInfo
+	OutFocalLengthInfo.FxFy = FVector2D(
+		float(CameraMatrix.at<double>(0, 0) / ImageSize.X),
+		float(CameraMatrix.at<double>(1, 1) / ImageSize.Y)
+	);
+
+	// DistortionInfo
+	{
+		FSphericalDistortionParameters SphericalParameters;
+
+		SphericalParameters.K1 = DistortionCoefficients.at<double>(0);
+		SphericalParameters.K2 = DistortionCoefficients.at<double>(1);
+		SphericalParameters.P1 = DistortionCoefficients.at<double>(2);
+		SphericalParameters.P2 = DistortionCoefficients.at<double>(3);
+		SphericalParameters.K3 = DistortionCoefficients.at<double>(4);
+
+		USphericalLensModel::StaticClass()->GetDefaultObject<ULensModel>()->ToArray(
+			SphericalParameters,
+			OutDistortionInfo.Parameters
+		);
+	}
+
+	// ImageCenterInfo
+	OutImageCenterInfo.PrincipalPoint = FVector2D(
+		float(CameraMatrix.at<double>(0, 2) / ImageSize.X),
+		float(CameraMatrix.at<double>(1, 2) / ImageSize.Y)
+	);
+
+	// Lens Model
+	OutLensModel = USphericalLensModel::StaticClass();
+
+	return true;
+
+#else
+	OutErrorMessage = LOCTEXT("OpenCVRequired", "OpenCV is required");
+	return false;
+#endif //WITH_OPENCV
+}
+
+void UCameraLensDistortionAlgoPoints::OnDistortionSavedToLens()
+{
+	// Since the distortion info was saved, there is no further use for the current samples.
+	ClearCalibrationRows();
+}
+
+AActor* UCameraLensDistortionAlgoPoints::FindFirstCalibrator() const
+{
+	// We find the first UCalibrationPointComponent object and return its actor owner.
+	if (!LensDistortionTool.IsValid())
+	{
+		return nullptr;
+	}
+
+	const FCameraCalibrationStepsController* StepsController = LensDistortionTool->GetCameraCalibrationStepsController();
+
+	if (!StepsController)
+	{
+		return nullptr;
+	}
+
+	const UWorld* World = StepsController->GetWorld();
+	const EObjectFlags ExcludeFlags = RF_ClassDefaultObject; // We don't want the calibrator CDOs.
+
+	for (TObjectIterator<UCalibrationPointComponent> It(ExcludeFlags, true, EInternalObjectFlags::PendingKill); It; ++It)
+	{
+		AActor* Actor = It->GetOwner();
+
+		if (Actor && (Actor->GetWorld() == World))
+		{
+			return Actor;
+		}
+	}
+
+	return nullptr;
+}
+
+void UCameraLensDistortionAlgoPoints::SetCalibrator(AActor* InCalibrator)
+{
+	Calibrator = InCalibrator;
+
+	// Update the list of points
+	CalibratorPoints.Empty();
+
+	if (!Calibrator.IsValid())
+	{
+		return;
+	}
+
+	TArray<UCalibrationPointComponent*, TInlineAllocator<NumInlineAllocations>> CalibrationPoints;
+	Calibrator->GetComponents<UCalibrationPointComponent, TInlineAllocator<NumInlineAllocations>>(CalibrationPoints);
+
+	for (const UCalibrationPointComponent* CalibrationPoint : CalibrationPoints)
+	{
+		TArray<FString> PointNames;
+
+		CalibrationPoint->GetNamespacedPointNames(PointNames);
+
+		for (FString& PointName : PointNames)
+		{
+			CalibratorPoints.Add(MakeShared<FCalibratorPointData>(PointName));
+		}
+	}
+
+	if (CalibratorPoints.Num() > 0)
+	{
+		CurrentCalibratorPoint = CalibratorPoints[0];
+	}
+	else
+	{
+		CurrentCalibratorPoint = nullptr;
+	}
+}
+
+bool UCameraLensDistortionAlgoPoints::AdvanceCalibratorPoint()
+{
+	if (!CurrentCalibratorPoint.IsValid())
+	{
+		return false;
+	}
+
+	for (int32 PointIdx = 0; PointIdx < CalibratorPoints.Num(); PointIdx++)
+	{
+		if (CalibratorPoints[PointIdx]->Name == CurrentCalibratorPoint->Name)
+		{
+			const int32 NextIdx = (PointIdx + 1) % CalibratorPoints.Num();
+			CurrentCalibratorPoint = CalibratorPoints[NextIdx];
+
+			// return true if we wrapped around (NextIdx is zero)
+			return !NextIdx;
+		}
+	}
+
+	return false;
+}
+
+bool UCameraLensDistortionAlgoPoints::GetCalibratorPointCacheFromName(const FString& Name, FCalibratorPointCache& CalibratorPointCache) const
+{
+	CalibratorPointCache.bIsValid = false;
+
+	if (!Calibrator.IsValid())
+	{
+		return false;
+	}
+
+	TArray<UCalibrationPointComponent*, TInlineAllocator<4>> CalibrationPoints;
+	Calibrator->GetComponents<UCalibrationPointComponent, TInlineAllocator<4>>(CalibrationPoints);
+
+	for (const UCalibrationPointComponent* CalibrationPoint : CalibrationPoints)
+	{
+		if (CalibrationPoint->GetWorldLocation(Name, CalibratorPointCache.Point3d))
+		{
+			CalibratorPointCache.bIsValid = true;
+			CalibratorPointCache.Name = Name;
+			return true;
+		}
+	}
+
+	return false;
+}
+
+bool UCameraLensDistortionAlgoPoints::ValidateNewRow(TSharedPtr<FCalibrationRowData>& Row, FText& OutErrorMessage) const
+{
+	const FCameraCalibrationStepsController* StepsController = LensDistortionTool->GetCameraCalibrationStepsController();
+
+	if (!StepsController)
+	{
+		OutErrorMessage = LOCTEXT("InvalidStepsController", "Invalid StepsController");
+		return false;
+	}
+
+	if (!Row.IsValid())
+	{
+		OutErrorMessage = LOCTEXT("InvalidRowPointer", "Invalid row pointer");
+		return false;
+	}
+
+	// The remaining checks all validate that data in the new row matches data in the first row.
+	// If this is the first row, those checks are unnecessary. 
+	if (CalibrationRows.Num() == 0)
+	{
+		return true;
+	}
+
+	// Same camera as before
+	const TSharedPtr<FCalibrationRowData>& FirstRow = CalibrationRows[0];
+
+	if (FirstRow->CameraData.UniqueId != Row->CameraData.UniqueId)
+	{
+		OutErrorMessage = LOCTEXT("CameraChangedDuringTheTest", "Camera changed during the test");
+		return false;
+	}
+
+	// Same calibrator as before
+	if (FirstRow->CameraData.CalibratorUniqueId != Row->CameraData.CalibratorUniqueId)
+	{
+		OutErrorMessage = LOCTEXT("CalibratorChangedDuringTheTest", "Calibrator changed during the test");
+		return false;
+	}
+
+	// FZ inputs are always valid, no need to verify them. They could be coming from LiveLink or fallback to a default one
+
+	//@todo Focus and zoom did not change much (i.e. inputs to distortion and nodal offset). 
+	//      Threshold for physical units should differ from normalized encoders.
+
+	return true;
+}
+
+void UCameraLensDistortionAlgoPoints::ClearCalibrationRows()
+{
+	CalibrationRows.Empty();
+
+	if (CalibrationListView.IsValid())
+	{
+		CalibrationListView->RequestListRefresh();
+	}
+
+	// Reset the pattern index
+	CurrentPatternIndex = 0;
+
+	// Reset the next calibrator point to the first one
+	if (CalibratorPoints.Num() > 0)
+	{
+		CurrentCalibratorPoint = CalibratorPoints[0];
+	}
+	else
+	{
+		CurrentCalibratorPoint = nullptr;
+	}
+}
+
+void UCameraLensDistortionAlgoPoints::DeleteSelectedCalibrationRows()
+{
+	const TArray<TSharedPtr<FCalibrationRowData>> SelectedItems = CalibrationListView->GetSelectedItems();
+
+	TSet<uint32> PatternsToDelete;
+
+	for (const TSharedPtr<FCalibrationRowData>& Row : SelectedItems)
+	{
+		PatternsToDelete.Add(Row->PatternIndex);
+	}
+
+	if (PatternsToDelete.Contains(CurrentPatternIndex))
+	{
+		if (CalibratorPoints.Num() > 0)
+		{
+			CurrentCalibratorPoint = CalibratorPoints[0];
+		}
+		else
+		{
+			CurrentCalibratorPoint = nullptr;
+		}
+	}
+
+	for (const TSharedPtr<FCalibrationRowData>& SelectedItem : SelectedItems)
+	{
+		CalibrationRows.Remove(SelectedItem);
+	}
+
+	CalibrationListView->RequestListRefresh();
+}
+
+void UCameraLensDistortionAlgoPoints::OnCalibrationRowSelectionChanged(TSharedPtr<FCalibrationRowData> SelectedRow, ESelectInfo::Type SelectInfo)
+{
+	if (SelectInfo != ESelectInfo::Type::Direct)
+	{
+		TSet<uint32> PatternsToSelect;
+
+		for (const TSharedPtr<FCalibrationRowData>& Row : CalibrationRows)
+		{
+			if (CalibrationListView->IsItemSelected(Row))
+			{
+				PatternsToSelect.Add(Row->PatternIndex);
+			}
+		}
+
+		TArray<TSharedPtr<FCalibrationRowData>> ItemsToSelect;
+
+		// Select all of the rows that have a matching pattern index (so that whole patterns are selected as groups)
+		for (const TSharedPtr<FCalibrationRowData>& Row : CalibrationRows)
+		{
+			if (PatternsToSelect.Contains(Row->PatternIndex))
+			{
+				ItemsToSelect.Add(Row);
+			}
+		}
+
+		CalibrationListView->SetItemSelection(ItemsToSelect, true);
+	}
+}
+
+TSharedRef<SWidget> UCameraLensDistortionAlgoPoints::BuildCalibrationDevicePickerWidget()
+{
+	return SNew(SFilterableActorPicker)
+		.OnSetObject_Lambda([&](const FAssetData& AssetData) -> void
+		{
+			if (AssetData.IsValid())
+			{
+				SetCalibrator(Cast<AActor>(AssetData.GetAsset()));
+			}
+		})
+		.OnShouldFilterAsset_Lambda([&](const FAssetData& AssetData) -> bool
+		{
+			const AActor* Actor = Cast<AActor>(AssetData.GetAsset());
+
+			if (!Actor)
+			{
+				return false;
+			}
+
+			TArray<UCalibrationPointComponent*, TInlineAllocator<NumInlineAllocations>> CalibrationPoints;
+			Actor->GetComponents<UCalibrationPointComponent, TInlineAllocator<NumInlineAllocations>>(CalibrationPoints);
+
+			return (CalibrationPoints.Num() > 0);
+		})
+		.ActorAssetData_Lambda([&]() -> FAssetData
+		{
+			return FAssetData(Calibrator.Get(), true);
+		});
+}
+
+TSharedRef<SWidget> UCameraLensDistortionAlgoPoints::BuildCurrentCalibratorPointLabel()
+{
+	return SNew(STextBlock)
+		.Text_Lambda([&]() -> FText
+		{
+			if (CurrentCalibratorPoint.IsValid())
+			{
+				return FText::FromString(CurrentCalibratorPoint->Name);
+			}
+
+			return LOCTEXT("NoCurrentCalibrator", "None");
+		})
+		.ColorAndOpacity(FLinearColor::White);
+}
+
+TSharedRef<SWidget> UCameraLensDistortionAlgoPoints::BuildCurrentFocalLengthWidget()
+{
+	return SNew(SNumericEntryBox<float>)
+		.Value_Lambda([&]() { return CurrentFocalLengthMM; })
+		.ToolTipText(LOCTEXT("CurrenFocalLength", "Current Focal Length (in mm)"))
+		.OnValueChanged_Lambda([&](float InValue)
+		{
+			if (InValue < 1.0f)
+			{
+				InValue = 1.0f;
+			}
+			CurrentFocalLengthMM = InValue;
+		});
+}
+
+TSharedRef<SWidget> UCameraLensDistortionAlgoPoints::BuildCalibrationPointsTable()
+{
+	CalibrationListView = SNew(SListView<TSharedPtr<FCalibrationRowData>>)
+		.ItemHeight(24)
+		.ListItemsSource(&CalibrationRows)
+		.OnGenerateRow_Lambda([&](TSharedPtr<FCalibrationRowData> InItem, const TSharedRef<STableViewBase>& OwnerTable) -> TSharedRef<ITableRow>
+		{
+			return SNew(CameraLensDistortionAlgoPoints::SCalibrationRowGenerator, OwnerTable).CalibrationRowData(InItem);
+		})
+		.SelectionMode(ESelectionMode::Multi)
+		.OnSelectionChanged_Lambda([&](TSharedPtr<FCalibrationRowData> SelectedRow, ESelectInfo::Type SelectInfo)
+		{
+			OnCalibrationRowSelectionChanged(SelectedRow, SelectInfo);
+		})
+		.OnKeyDownHandler_Lambda([&](const FGeometry& Geometry, const FKeyEvent& KeyEvent) -> FReply
+		{
+			if (!CalibrationListView.IsValid())
+			{
+				return FReply::Unhandled();
+			}
+
+			if (KeyEvent.GetKey() == EKeys::Delete)
+			{
+				DeleteSelectedCalibrationRows();
+				return FReply::Handled();
+			}
+			else if (KeyEvent.GetModifierKeys().IsControlDown() && (KeyEvent.GetKey() == EKeys::A))
+			{
+				// Select all items
+				CalibrationListView->SetItemSelection(CalibrationRows, true);
+				return FReply::Handled();
+			}
+			else if (KeyEvent.GetKey() == EKeys::Escape)
+			{
+				// Deselect all items
+				CalibrationListView->ClearSelection();
+				return FReply::Handled();
+			}
+
+			return FReply::Unhandled();
+		})
+		.HeaderRow
+		(
+			SNew(SHeaderRow)
+
+			+ SHeaderRow::Column("Pattern")
+			.DefaultLabel(LOCTEXT("Pattern", "Pattern"))
+			.FillWidth(0.15f)
+
+			+ SHeaderRow::Column("Name")
+			.DefaultLabel(LOCTEXT("Name", "Name"))
+			.FillWidth(0.25f)
+
+			+ SHeaderRow::Column("Point2D")
+			.DefaultLabel(LOCTEXT("Point2D", "Point2D"))
+			.FillWidth(0.3f)
+
+			+ SHeaderRow::Column("Point3D")
+			.DefaultLabel(LOCTEXT("Point3D", "Point3D"))
+			.FillWidth(0.3f)
+		);
+
+	return CalibrationListView.ToSharedRef();
+}
+
+TSharedRef<SWidget> UCameraLensDistortionAlgoPoints::BuildCalibrationActionButtons()
+{
+	return SNew(SHorizontalBox)
+
+		+ SHorizontalBox::Slot() // Button to clear all rows
+		.AutoWidth()
+		[
+			SNew(SButton)
+			.Text(LOCTEXT("ClearAll", "Clear All"))
+			.HAlign(HAlign_Center)
+			.VAlign(VAlign_Center)
+			.OnClicked_Lambda([&]() -> FReply
+			{
+				ClearCalibrationRows();
+				return FReply::Handled();
+			})
+		]
+	;
+}
+
+TSharedRef<SWidget> UCameraLensDistortionAlgoPoints::BuildHelpWidget()
+{
+	return SNew(STextBlock)
+		.Text(LOCTEXT("LensDistortionPointsHelp",
+			"This Lens Distortion algorithm is based on capturing at least 4 different views of a calibration pattern.\n\n"
+			"The camera and/or the calibrator may be moved for each. To capture, simply click the simulcam\n"
+			"viewport. You can optionally right-click the simulcam viewport to pause it and ensure it will be\n"
+			"a sharp capture (if the media source supports pause).\n\n"
+
+			"This will require the selection of a calibrator actor that exists in the scene, and is configured\n"
+			"to have several calibration points components that correspond to physical features that are easy to detect.\n\n"
+
+			"Each click will add a row to the table with the 3d and 2d points of the current calibration point. The tool\n"
+			"will validate that you capture at least 4 patterns, and that each pattern has the same number of calibration\n"
+			"points in it. It is important that captures sweep the viewport area as much as possible, so that the\n"
+			"calibration includes samples from all of its regions (otherwise the calibration may be skewed towards a\n"
+			"particular region and not be very accurate in others).\n\n"
+
+			"When you are done capturing, click the 'Add To Lens Distortion Calibration' button to run the calibration\n"
+			"algorithm, which calculates the distortion parameters k1,k2,p1,p2,k3 and adds them to the lens file.\n\n"
+
+			"A dialog window will show you the reprojection error of the calibration, and you can decide to add the\n"
+			"data to the lens file or not.\n"
+		));
+
+}
+
+#undef LOCTEXT_NAMESPACE
diff --git a/Engine/Plugins/VirtualProduction/CameraCalibration/Source/CameraCalibrationEditor/Private/Calibrators/CameraLensDistortionAlgoPoints.h b/Engine/Plugins/VirtualProduction/CameraCalibration/Source/CameraCalibrationEditor/Private/Calibrators/CameraLensDistortionAlgoPoints.h
new file mode 100644
--- /dev/null
+++ b/Engine/Plugins/VirtualProduction/CameraCalibration/Source/CameraCalibrationEditor/Private/Calibrators/CameraLensDistortionAlgoPoints.h
@@ -0,0 +1,201 @@
+// Copyright Epic Games, Inc. All Rights Reserved.
+
+#pragma once
+
+#include "CameraLensDistortionAlgo.h"
+
+#include "LensFile.h"
+
+#include "CameraLensDistortionAlgoPoints.generated.h"
+
+template <typename ItemType>
+class SListView;
+
+namespace CameraLensDistortionAlgoPoints
+{
+	class SCalibrationRowGenerator;
+};
+
+/**
+ * Implements a lens distortion calibration algorithm. It uses 3d points (UCalibrationPointComponent)
+ * specified in a selectable calibrator with features that the user can identify by clicking on the
+ * simulcam viewport, and after enough points have been identified, it can calculate the lens distortion
+ * and camera intrinsics that need to be applied to the associated camera in order to align its CG with 
+ * the live action media plate.
+ *
+ * Limitations:
+ * - Only supports Brown–Conrady model lens model (FSphericalDistortionParameters)
+ * - The camera or the lens should not be moved during the calibration of each lens distortion sample.
+ */
+UCLASS()
+class UCameraLensDistortionAlgoPoints : public UCameraLensDistortionAlgo
+{
+	GENERATED_BODY()
+
+public:
+	//~ Begin UCameraLensDistortionAlgo
+	virtual void Initialize(ULensDistortionTool* InLensDistortionTool) override;
+	virtual void Shutdown() override;
+	virtual void Tick(float DeltaTime) override;
+	virtual bool OnViewportClicked(const FGeometry& MyGeometry, const FPointerEvent& MouseEvent) override;
+	virtual TSharedRef<SWidget> BuildUI() override;
+	virtual FName FriendlyName() const override { return TEXT("Lens Distortion Points Method"); };
+
+	virtual bool GetLensDistortion(
+		float& OutFocus,
+		float& OutZoom,
+		FDistortionInfo& OutDistortionInfo,
+		FFocalLengthInfo& OutFocalLengthInfo,
+		FImageCenterInfo& OutImageCenterInfo,
+		TSubclassOf<ULensModel>& OutLensModel,
+		double& OutError,
+		FText& OutErrorMessage) override;
+
+	virtual void OnDistortionSavedToLens() override;
+	virtual TSharedRef<SWidget> BuildHelpWidget() override;
+	//~ End UCameraLensDistortionAlgo
+
+protected:
+
+	// SCalibrationRowGenerator will need access to the row structures below.
+	friend class CameraLensDistortionAlgoPoints::SCalibrationRowGenerator;
+
+	/** Holds item information a given calibrator point in the calibrator model */
+	struct FCalibratorPointData
+	{
+		FCalibratorPointData(FString& InName)
+			: Name(InName)
+		{}
+
+		/** Name of the calibrator 3d point, as defined in the mesh */
+		FString Name;
+	};
+
+	/** Holds information of the identified calibrator 2d point for a given sample of a 2d-3d correlation */
+	struct FCalibratorPointCache
+	{
+		// True if the rest of the contents are valid.
+		bool bIsValid;
+
+		// The name of the 3d calibrator point, as defined in the mesh
+		FString Name;
+
+		// 3D location of the calibrator point
+		FVector Point3d;
+
+		// 2D location of the calibrator point
+		FVector2D Point2d;
+	};
+
+	/** Holds information of the camera pose for a given sample of a 2d-3d correlation */
+	struct FCameraDataCache
+	{
+		// True if the rest of the contents are valid.
+		bool bIsValid;
+
+		// The unique id of the camera object. Used to detect camera selection changes during a calibration session.
+		uint32 UniqueId;
+
+		// The unique id of the calibrator object. Used to detect calibrator selection changes during a calibration session.
+		uint32 CalibratorUniqueId;
+
+		// The data used to evaluate the lens data in the camera for this sample
+		FLensFileEvalData LensFileEvalData;
+	};
+
+	/** Holds information of the calibrator 3d point for a given sample of a 2d-3d correlation */
+	struct FCalibrationRowData
+	{
+		// The 3D-2D point correspondences for each of the calibrator points in the chosen calibrator for a single pose
+		FCalibratorPointCache CalibratorPointData;
+
+		// Holds information of the camera data for this sample
+		FCameraDataCache CameraData;
+
+		// Index of the calibration pattern this row is associated with
+		uint32 PatternIndex;
+	};
+
+protected:
+
+	/** The lens distortion tool */
+	TWeakObjectPtr<ULensDistortionTool> LensDistortionTool;
+
+	/** The currently selected calibrator object. It is expected to contain one or more UCalibrationPointComponent in it */
+	TWeakObjectPtr<AActor> Calibrator;
+
+	/** Lists the calibrator points found in the selected calibrator object */
+	TArray<TSharedPtr<FCalibratorPointData>> CalibratorPoints;
+
+	/** The current calibrator point that will be used to generate the next calibration row */
+	TSharedPtr<FCalibratorPointData> CurrentCalibratorPoint;
+
+	/** Rows source for the CalibrationListView */
+	TArray<TSharedPtr<FCalibrationRowData>> CalibrationRows;
+
+	/** Displays the list of calibration points that will be used to calculate the lens distortion */
+	TSharedPtr<SListView<TSharedPtr<FCalibrationRowData>>> CalibrationListView;
+
+	/** Caches the last calibrator point 3d location.  Will hold last value before the lens distortion tool is paused */
+	TArray<FCalibratorPointCache> LastCalibratorPoints;
+
+	/** Caches the last camera data.  Will hold last value before the lens distortion tool is paused */
+	FCameraDataCache LastCameraData;
+
+	/** Instructs the inline allocation policy how many elements to allocate in a single allocation */
+	static constexpr uint32 NumInlineAllocations = 32;
+
+	/** The current index of the calibration pattern being captured, used for grouping calibration rows together into meaningful patterns */
+	uint32 CurrentPatternIndex = 0;
+
+	/** The value of the current focal length of the lens (in mm), used to set a guess for the camera intrinsics fx/fy */
+	float CurrentFocalLengthMM = 0.0f;
+
+protected:
+
+	/** Returns the first calibrator object in the scene that it can find */
+	AActor* FindFirstCalibrator() const;
+
+	/** Sets the calibrator object to be used. Updates the selection picker. */
+	void SetCalibrator(AActor* InCalibrator);
+
+	/** Selects the next available UCalibrationPointComponent of the currently selected calibrator object. Returns true when it wraps around */
+	bool AdvanceCalibratorPoint();
+
+	/**
+	 * Retrieves by name the calibration point data of the currently selected calibrator.
+	 *
+	 * @param Name The name of the point (namespaced if it is a subpoint).
+	 * @param CalibratorPointCache This data structure will be populated with the information found.
+	 *
+	 * @return True if successful.
+	 */
+	bool GetCalibratorPointCacheFromName(const FString& Name, FCalibratorPointCache& CalibratorPointCache) const;
+
+	/** Validates a new calibration point to determine if it should be added as a new sample row */
+	bool ValidateNewRow(TSharedPtr<FCalibrationRowData>& Row, FText& OutErrorMessage) const;
+
+	/** Clears the list of calibration sample rows */
+	void ClearCalibrationRows();
+
+	/** Deletes the selected calibration rows */
+	void DeleteSelectedCalibrationRows();
+
+	/** Updates the selection in the list of calibration rows so that patterns are always grouped together */
+	void OnCalibrationRowSelectionChanged(TSharedPtr<FCalibrationRowData> SelectedRow, ESelectInfo::Type SelectInfo);
+
+	/** Builds the UI of the calibration device picker */
+	TSharedRef<SWidget> BuildCalibrationDevicePickerWidget();
+
+	/** Buils the UI for selecting the current calibrator point */
+	TSharedRef<SWidget> BuildCurrentCalibratorPointLabel();
+
+	/** Buils the UI for setting the current focal length (in mm) */
+	TSharedRef<SWidget> BuildCurrentFocalLengthWidget();
+
+	/** Builds the UI of the calibration points table */
+	TSharedRef<SWidget> BuildCalibrationPointsTable();
+
+	/** Builds the UI for the action buttons (RemoveLast, ClearAll) */
+	TSharedRef<SWidget> BuildCalibrationActionButtons();
+};
