Fixed duplicate localization keys in VP-IO plugins

diff --git a/Engine/Plugins/VirtualProduction/CameraCalibration/Source/CameraCalibrationEditor/Private/AssetEditor/Curves/LensFocalLengthCurveModel.cpp b/Engine/Plugins/VirtualProduction/CameraCalibration/Source/CameraCalibrationEditor/Private/AssetEditor/Curves/LensFocalLengthCurveModel.cpp
--- a/Engine/Plugins/VirtualProduction/CameraCalibration/Source/CameraCalibrationEditor/Private/AssetEditor/Curves/LensFocalLengthCurveModel.cpp
+++ b/Engine/Plugins/VirtualProduction/CameraCalibration/Source/CameraCalibrationEditor/Private/AssetEditor/Curves/LensFocalLengthCurveModel.cpp
@@ -159,9 +159,9 @@ FText FLensFocalLengthCurveModel::GetValueLabel() const
 {
 	if (ParameterIndex == INDEX_NONE)
 	{
-		return LOCTEXT("ZoomValueLabel", "(mm)");
+		return LOCTEXT("ZoomValueLabelMM", "(mm)");
 	}
-	return LOCTEXT("ZoomValueLabel", "(normalized)");
+	return LOCTEXT("ZoomValueLabelNormalized", "(normalized)");
 }
 
 FText FLensFocalLengthCurveModel::GetValueUnitSuffixLabel() const
diff --git a/Engine/Plugins/VirtualProduction/CameraCalibration/Source/CameraCalibrationEditor/Private/AssetEditor/SCameraCalibrationRemovePointDialog.cpp b/Engine/Plugins/VirtualProduction/CameraCalibration/Source/CameraCalibrationEditor/Private/AssetEditor/SCameraCalibrationRemovePointDialog.cpp
--- a/Engine/Plugins/VirtualProduction/CameraCalibration/Source/CameraCalibrationEditor/Private/AssetEditor/SCameraCalibrationRemovePointDialog.cpp
+++ b/Engine/Plugins/VirtualProduction/CameraCalibration/Source/CameraCalibrationEditor/Private/AssetEditor/SCameraCalibrationRemovePointDialog.cpp
@@ -236,11 +236,11 @@ void FFocusRemoveLensDataListItem::OnCheckStateChanged(ECheckBoxState NewState)
 
 TSharedRef<ITableRow> FZoomRemoveLensDataListItem::MakeTreeRowWidget(const TSharedRef<STableViewBase>& InOwnerTable)
 {
-	FText EntryLabel = LOCTEXT("ZoomLabel", "Zoom: ");
+	FText EntryLabel = LOCTEXT("ZoomLabelIncomplete", "Zoom: ");
 
 	if (!WeakParent.IsValid())
 	{
-		EntryLabel = FText::Format(LOCTEXT("ZoomLabel", "{0}. Focus: {1}, Zoom: "), FText::FromName(PointName), Focus);
+		EntryLabel = FText::Format(LOCTEXT("ZoomLabelComplete", "{0}. Focus: {1}, Zoom: "), FText::FromName(PointName), Focus);
 	}
 		
 	return SNew(SRemoveLensDataItem, InOwnerTable, SharedThis(this))
diff --git a/Engine/Plugins/VirtualProduction/CameraCalibration/Source/CameraCalibrationEditor/Private/AssetEditor/SCameraCalibrationSteps.cpp b/Engine/Plugins/VirtualProduction/CameraCalibration/Source/CameraCalibrationEditor/Private/AssetEditor/SCameraCalibrationSteps.cpp
--- a/Engine/Plugins/VirtualProduction/CameraCalibration/Source/CameraCalibrationEditor/Private/AssetEditor/SCameraCalibrationSteps.cpp
+++ b/Engine/Plugins/VirtualProduction/CameraCalibration/Source/CameraCalibrationEditor/Private/AssetEditor/SCameraCalibrationSteps.cpp
@@ -365,7 +365,7 @@ TSharedRef<SWidget> SCameraCalibrationSteps::BuildOverlayWidget()
 					return FText::FromName(*OverlayComboBox->GetSelectedItem());
 				}
 
-				return LOCTEXT("InvalidComboOption", "None");
+				return LOCTEXT("NoneComboOption", "None");
 			})
 		];
 
diff --git a/Engine/Plugins/VirtualProduction/CameraCalibration/Source/CameraCalibrationEditor/Private/AssetEditor/SLensEvaluation.cpp b/Engine/Plugins/VirtualProduction/CameraCalibration/Source/CameraCalibrationEditor/Private/AssetEditor/SLensEvaluation.cpp
--- a/Engine/Plugins/VirtualProduction/CameraCalibration/Source/CameraCalibrationEditor/Private/AssetEditor/SLensEvaluation.cpp
+++ b/Engine/Plugins/VirtualProduction/CameraCalibration/Source/CameraCalibrationEditor/Private/AssetEditor/SLensEvaluation.cpp
@@ -518,7 +518,7 @@ TSharedRef<SWidget> SLensEvaluation::MakeEvaluatedFIZWidget() const
 					{
 						const float FOV = FMath::RadiansToDegrees(2.f * FMath::Atan(LensFile->LensInfo.SensorDimensions.X / (2.f * CachedLiveLinkData.EvaluatedZoom.GetValue())));
 						const FText ValueText = FText::AsNumber(FOV, &FloatOptions);
-						return FText::Format(LOCTEXT("PhysicalUnitsZoomValue", "{0} deg"), ValueText);
+						return FText::Format(LOCTEXT("PhysicalUnitsFOVValue", "{0} deg"), ValueText);
 					}
 					return LOCTEXT("UndefinedValue", "N/A");
 				}))
@@ -656,7 +656,7 @@ TSharedRef<SWidget> SLensEvaluation::MakeIntrinsicsWidget() const
 					{
 						const FText CxText = FText::AsNumber(CachedImageCenter.PrincipalPoint.X, &FloatOptions);
 						const FText CyText = FText::AsNumber(CachedImageCenter.PrincipalPoint.Y, &FloatOptions);
-						return FText::Format(LOCTEXT("LocationOffsetValue", "({0}, {1})"), CxText, CyText);
+						return FText::Format(LOCTEXT("PrincipalPointValue", "({0}, {1})"), CxText, CyText);
 					}
 					return LOCTEXT("UndefinedValue", "N/A");
 				}))
@@ -678,7 +678,7 @@ TSharedRef<SWidget> SLensEvaluation::MakeIntrinsicsWidget() const
 					{
 						const FText FxText = FText::AsNumber(CachedFocalLengthInfo.FxFy.X, &FloatOptions);
 						const FText FyText = FText::AsNumber(CachedFocalLengthInfo.FxFy.Y, &FloatOptions);
-						return FText::Format(LOCTEXT("LocationOffsetValue", "({0}, {1})"), FxText, FyText);
+						return FText::Format(LOCTEXT("FxFyValue", "({0}, {1})"), FxText, FyText);
 					}
 					return LOCTEXT("UndefinedValue", "N/A");
 				}))
@@ -755,7 +755,7 @@ TSharedRef<SWidget> SLensEvaluation::MakeNodalOffsetWidget() const
 						const FText RotXText = FText::AsNumber(CachedNodalOffset.RotationOffset.Rotator().GetComponentForAxis(EAxis::X), &FloatOptions);
 						const FText RotYText = FText::AsNumber(CachedNodalOffset.RotationOffset.Rotator().GetComponentForAxis(EAxis::Y), &FloatOptions);
 						const FText RotZText = FText::AsNumber(CachedNodalOffset.RotationOffset.Rotator().GetComponentForAxis(EAxis::Z), &FloatOptions);
-						return FText::Format(LOCTEXT("LocationOffsetValue", "({0}, {1}, {2})"), RotXText, RotYText, RotZText);
+						return FText::Format(LOCTEXT("RotationOffsetValue", "({0}, {1}, {2})"), RotXText, RotYText, RotZText);
 					}
 					return LOCTEXT("UndefinedValue", "N/A");
 				}))
diff --git a/Engine/Plugins/VirtualProduction/CameraCalibration/Source/CameraCalibrationEditor/Private/Calibrators/CameraNodalOffsetAlgoPoints.cpp b/Engine/Plugins/VirtualProduction/CameraCalibration/Source/CameraCalibrationEditor/Private/Calibrators/CameraNodalOffsetAlgoPoints.cpp
--- a/Engine/Plugins/VirtualProduction/CameraCalibration/Source/CameraCalibrationEditor/Private/Calibrators/CameraNodalOffsetAlgoPoints.cpp
+++ b/Engine/Plugins/VirtualProduction/CameraCalibration/Source/CameraCalibrationEditor/Private/Calibrators/CameraNodalOffsetAlgoPoints.cpp
@@ -942,7 +942,7 @@ bool UCameraNodalOffsetAlgoPoints::BasicCalibrationChecksPass(const TArray<TShar
 	// Enough points
 	if (Rows.Num() < 4)
 	{
-		OutErrorMessage = LOCTEXT("NotEnoughSamples", "At least 4 correspondence points are required");
+		OutErrorMessage = LOCTEXT("LessThanFourSamples", "At least 4 correspondence points are required");
 		return false;
 	}
 
@@ -1963,7 +1963,7 @@ bool UCameraNodalOffsetAlgoPoints::ApplyNodalOffsetToCalibratorParent()
 
 		if (Row->CameraData.CalibratorParentUniqueId != ParentActor->GetUniqueID())
 		{
-			ErrorMessage = LOCTEXT("WrongCalibrator", "All rows must belong to the same calibrator parent");
+			ErrorMessage = LOCTEXT("WrongCalibratorParent", "All rows must belong to the same calibrator parent");
 			FMessageDialog::Open(EAppMsgType::Ok, ErrorMessage, &TitleError);
 
 			return false;
diff --git a/Engine/Plugins/VirtualProduction/DMX/DMXEngine/Source/DMXEditor/Private/DMXFixtureTypeSharedData.cpp b/Engine/Plugins/VirtualProduction/DMX/DMXEngine/Source/DMXEditor/Private/DMXFixtureTypeSharedData.cpp
--- a/Engine/Plugins/VirtualProduction/DMX/DMXEngine/Source/DMXEditor/Private/DMXFixtureTypeSharedData.cpp
+++ b/Engine/Plugins/VirtualProduction/DMX/DMXEngine/Source/DMXEditor/Private/DMXFixtureTypeSharedData.cpp
@@ -757,7 +757,7 @@ void FDMXFixtureTypeSharedData::PasteClipboardToFunctions(const TArray<int32>& F
 				{
 					if (UDMXEntityFixtureType* FixtureType = WeakFixtureType.Get())
 					{
-						const FScopedTransaction ModeNameTransaction = FScopedTransaction(LOCTEXT("PasteModeTransaction", "Pasted Fixture Type Mode"));
+						const FScopedTransaction ModeNameTransaction = FScopedTransaction(LOCTEXT("PasteModeTransaction", "Paste Fixture Type Mode"));
 
 						FixtureType->PreEditChange(UDMXEntityFixtureType::StaticClass()->FindPropertyByName(GET_MEMBER_NAME_STRING_CHECKED(UDMXEntityFixtureType, Modes)));
 						FixtureType->Modify();
diff --git a/Engine/Plugins/VirtualProduction/DMX/DMXEngine/Source/DMXEditor/Private/Widgets/SAssetPickerButton.cpp b/Engine/Plugins/VirtualProduction/DMX/DMXEngine/Source/DMXEditor/Private/Widgets/SAssetPickerButton.cpp
--- a/Engine/Plugins/VirtualProduction/DMX/DMXEngine/Source/DMXEditor/Private/Widgets/SAssetPickerButton.cpp
+++ b/Engine/Plugins/VirtualProduction/DMX/DMXEngine/Source/DMXEditor/Private/Widgets/SAssetPickerButton.cpp
@@ -79,7 +79,7 @@ void SAssetPickerButton::Construct(const FArguments& InArgs)
 			.ButtonColorAndOpacity( this, &SAssetPickerButton::OnGetWidgetBackground )
 			.OnClicked(this, &SAssetPickerButton::OnClickUse)
 			.ContentPadding(1.f)
-			.ToolTipText(NSLOCTEXT("GraphEditor", "ObjectGraphPin_Use_Tooltip", "Use asset browser selection"))
+			.ToolTipText(LOCTEXT("UseSelectionTooltip", "Use asset browser selection"))
 			[
 				SNew(SImage)
 				.ColorAndOpacity( this, &SAssetPickerButton::OnGetWidgetForeground )
@@ -97,7 +97,7 @@ void SAssetPickerButton::Construct(const FArguments& InArgs)
 			.ButtonColorAndOpacity( this, &SAssetPickerButton::OnGetWidgetBackground )
 			.OnClicked(this, &SAssetPickerButton::OnClickBrowse)
 			.ContentPadding(0)
-			.ToolTipText(NSLOCTEXT("GraphEditor", "ObjectGraphPin_Browse_Tooltip", "Browse"))
+			.ToolTipText(LOCTEXT("BrowseTooltip", "Browse"))
 			[
 				SNew(SImage)
 				.ColorAndOpacity( this, &SAssetPickerButton::OnGetWidgetForeground )
diff --git a/Engine/Plugins/VirtualProduction/DMX/DMXEngine/Source/DMXRuntime/Private/Library/DMXEntityFixturePatch.cpp b/Engine/Plugins/VirtualProduction/DMX/DMXEngine/Source/DMXRuntime/Private/Library/DMXEntityFixturePatch.cpp
--- a/Engine/Plugins/VirtualProduction/DMX/DMXEngine/Source/DMXRuntime/Private/Library/DMXEntityFixturePatch.cpp
+++ b/Engine/Plugins/VirtualProduction/DMX/DMXEngine/Source/DMXRuntime/Private/Library/DMXEntityFixturePatch.cpp
@@ -341,7 +341,7 @@ bool UDMXEntityFixturePatch::IsValidEntity(FText& OutReason) const
 	else if (!ParentFixtureTypeTemplate->Modes.IsValidIndex(ActiveMode))
 	{
 		OutReason = FText::Format(
-			LOCTEXT("InvalidReason_NoModesDefined", "'{0}' cannot be assigned as its parent Fixture Type '{1}' does no longer contain Active Mode."),
+			LOCTEXT("InvalidReason_NoActiveModes", "'{0}' cannot be assigned as its parent Fixture Type '{1}' does no longer contain Active Mode."),
 			FText::FromString(GetDisplayName()),
 			FText::FromString(ParentFixtureTypeTemplate->GetDisplayName())
 		);
diff --git a/Engine/Plugins/VirtualProduction/DMX/DMXPixelMapping/Source/DMXPixelMappingEditor/Private/DragDrop/DMXPixelMappingDragDropOp.cpp b/Engine/Plugins/VirtualProduction/DMX/DMXPixelMapping/Source/DMXPixelMappingEditor/Private/DragDrop/DMXPixelMappingDragDropOp.cpp
--- a/Engine/Plugins/VirtualProduction/DMX/DMXPixelMapping/Source/DMXPixelMappingEditor/Private/DragDrop/DMXPixelMappingDragDropOp.cpp
+++ b/Engine/Plugins/VirtualProduction/DMX/DMXPixelMapping/Source/DMXPixelMappingEditor/Private/DragDrop/DMXPixelMappingDragDropOp.cpp
@@ -68,7 +68,7 @@ TSharedRef<FDMXPixelMappingDragDropOp> FDMXPixelMappingDragDropOp::New(const FVe
 	Operation->SetDecoratorVisibility(false);
 
 	// Create a transaction for dragged templates
-	Operation->TransactionIndex = GEditor->BeginTransaction(FText::Format(LOCTEXT("DragDropTransaction", "PixelMapping: Add {0}|plural(one=Component, other=Components)"), InTemplates.Num()));
+	Operation->TransactionIndex = GEditor->BeginTransaction(FText::Format(LOCTEXT("DragDropTemplateTransaction", "PixelMapping: Add {0}|plural(one=Component, other=Components)"), InTemplates.Num()));
 
 	return Operation;
 }
@@ -86,7 +86,7 @@ TSharedRef<FDMXPixelMappingDragDropOp> FDMXPixelMappingDragDropOp::New(const FVe
 	Operation->SetDecoratorVisibility(false);
 
 	// Create a transaction for dragged components
-	Operation->TransactionIndex = GEditor->BeginTransaction(FText::Format(LOCTEXT("DragDropTransaction", "PixelMapping: Drag {0}|plural(one=Component, other=Components)"), InDraggedComponents.Num()));
+	Operation->TransactionIndex = GEditor->BeginTransaction(FText::Format(LOCTEXT("DragDropComponentTransaction", "PixelMapping: Drag {0}|plural(one=Component, other=Components)"), InDraggedComponents.Num()));
 
 	return Operation;
 }
diff --git a/Engine/Plugins/VirtualProduction/DMX/DMXProtocol/Source/DMXProtocolEditor/Private/DetailsCustomizations/DMXInputPortConfigCustomization.cpp b/Engine/Plugins/VirtualProduction/DMX/DMXProtocol/Source/DMXProtocolEditor/Private/DetailsCustomizations/DMXInputPortConfigCustomization.cpp
--- a/Engine/Plugins/VirtualProduction/DMX/DMXProtocol/Source/DMXProtocolEditor/Private/DetailsCustomizations/DMXInputPortConfigCustomization.cpp
+++ b/Engine/Plugins/VirtualProduction/DMX/DMXProtocol/Source/DMXProtocolEditor/Private/DetailsCustomizations/DMXInputPortConfigCustomization.cpp
@@ -244,7 +244,7 @@ void FDMXInputPortConfigCustomization::OnIPAddressSelected()
 {
 	FString SelectedIP = IPAddressEditWidget->GetSelectedIPAddress();
 
-	const FScopedTransaction Transaction(LOCTEXT("CommunicationTypeSelected", "DMX: Selected IP Address"));
+	const FScopedTransaction Transaction(LOCTEXT("IPAddressSelected", "DMX: Selected IP Address"));
 	
 	DeviceAddressHandle->NotifyPreChange();
 	ensure(DeviceAddressHandle->SetValue(SelectedIP) == FPropertyAccess::Success);	
diff --git a/Engine/Plugins/VirtualProduction/DMX/DMXProtocol/Source/DMXProtocolEditor/Private/DetailsCustomizations/DMXOutputPortConfigCustomization.cpp b/Engine/Plugins/VirtualProduction/DMX/DMXProtocol/Source/DMXProtocolEditor/Private/DetailsCustomizations/DMXOutputPortConfigCustomization.cpp
--- a/Engine/Plugins/VirtualProduction/DMX/DMXProtocol/Source/DMXProtocolEditor/Private/DetailsCustomizations/DMXOutputPortConfigCustomization.cpp
+++ b/Engine/Plugins/VirtualProduction/DMX/DMXProtocol/Source/DMXProtocolEditor/Private/DetailsCustomizations/DMXOutputPortConfigCustomization.cpp
@@ -285,7 +285,7 @@ void FDMXOutputPortConfigCustomization::OnIPAddressSelected()
 {
 	FString SelectedIP = IPAddressEditWidget->GetSelectedIPAddress();
 
-	const FScopedTransaction Transaction(LOCTEXT("CommunicationTypeSelected", "Selected DMX Output Port IP Address"));
+	const FScopedTransaction Transaction(LOCTEXT("IPAddressSelected", "Selected DMX Output Port IP Address"));
 
 	DeviceAddressHandle->NotifyPreChange();
 	ensure(DeviceAddressHandle->SetValue(SelectedIP) == FPropertyAccess::Success);
diff --git a/Engine/Plugins/VirtualProduction/RemoteControl/Source/RemoteControlProtocolWidgets/Private/ViewModels/ProtocolEntityViewModel.cpp b/Engine/Plugins/VirtualProduction/RemoteControl/Source/RemoteControlProtocolWidgets/Private/ViewModels/ProtocolEntityViewModel.cpp
--- a/Engine/Plugins/VirtualProduction/RemoteControl/Source/RemoteControlProtocolWidgets/Private/ViewModels/ProtocolEntityViewModel.cpp
+++ b/Engine/Plugins/VirtualProduction/RemoteControl/Source/RemoteControlProtocolWidgets/Private/ViewModels/ProtocolEntityViewModel.cpp
@@ -1,4 +1,4 @@
-﻿// Copyright Epic Games, Inc. All Rights Reserved.
+// Copyright Epic Games, Inc. All Rights Reserved.
 
 #include "ProtocolEntityViewModel.h"
 
@@ -146,7 +146,7 @@ bool FProtocolEntityViewModel::CanAddBinding(const FName& InProtocolName, FText&
 	}
 
 	// Remaining should be strings, enums
-	OutMessage = FText::Format(LOCTEXT("UnsupportedType", "Unsupported Type \"{0}\" for Protocol Binding"), GetProperty()->GetClass()->GetDisplayNameText());
+	OutMessage = FText::Format(LOCTEXT("UnsupportedTypeForBinding", "Unsupported Type \"{0}\" for Protocol Binding"), GetProperty()->GetClass()->GetDisplayNameText());
 	return false;
 }
 
diff --git a/Engine/Plugins/VirtualProduction/RemoteControl/Source/RemoteControlProtocolWidgets/Private/ViewModels/ProtocolRangeViewModel.cpp b/Engine/Plugins/VirtualProduction/RemoteControl/Source/RemoteControlProtocolWidgets/Private/ViewModels/ProtocolRangeViewModel.cpp
--- a/Engine/Plugins/VirtualProduction/RemoteControl/Source/RemoteControlProtocolWidgets/Private/ViewModels/ProtocolRangeViewModel.cpp
+++ b/Engine/Plugins/VirtualProduction/RemoteControl/Source/RemoteControlProtocolWidgets/Private/ViewModels/ProtocolRangeViewModel.cpp
@@ -7,7 +7,7 @@
 #include "PropertyHandle.h"
 #include "Misc/DefaultValueHelper.h"
 
-#define LOCTEXT_NAMESPACE "ProtocolBindingViewModel"
+#define LOCTEXT_NAMESPACE "ProtocolRangeViewModel"
 
 TMap<FProtocolRangeViewModel::EValidity, FText> FProtocolRangeViewModel::ValidityMessages =
 {
diff --git a/Engine/Plugins/VirtualProduction/RemoteControl/Source/RemoteControlUI/Private/RemoteControlUIModule.cpp b/Engine/Plugins/VirtualProduction/RemoteControl/Source/RemoteControlUI/Private/RemoteControlUIModule.cpp
--- a/Engine/Plugins/VirtualProduction/RemoteControl/Source/RemoteControlUI/Private/RemoteControlUIModule.cpp
+++ b/Engine/Plugins/VirtualProduction/RemoteControl/Source/RemoteControlUI/Private/RemoteControlUIModule.cpp
@@ -825,11 +825,11 @@ FText FRemoteControlUIModule::GetExposePropertyButtonText(const FRCExposesProper
 {
 	if (GetPropertyExposeStatus(InPropertyArgs) == EPropertyExposeStatus::Exposed)
 	{
-		return LOCTEXT("ExposePropertyToolTip", "Unexpose property");
+		return LOCTEXT("ExposePropertyText", "Unexpose property");
 	}
 	else
 	{
-		return LOCTEXT("UnexposePropertyToolTip", "Expose property");
+		return LOCTEXT("UnexposePropertyText", "Expose property");
 	}
 }
 
diff --git a/Engine/Plugins/VirtualProduction/RemoteControl/Source/RemoteControlUI/Private/UI/SRemoteControlPanel.cpp b/Engine/Plugins/VirtualProduction/RemoteControl/Source/RemoteControlUI/Private/UI/SRemoteControlPanel.cpp
--- a/Engine/Plugins/VirtualProduction/RemoteControl/Source/RemoteControlUI/Private/UI/SRemoteControlPanel.cpp
+++ b/Engine/Plugins/VirtualProduction/RemoteControl/Source/RemoteControlUI/Private/UI/SRemoteControlPanel.cpp
@@ -634,7 +634,7 @@ TSharedRef<SWidget> SRemoteControlPanel::CreateExposeButton()
 		
 		CreatePickerSubMenu(
 			LOCTEXT("ActorFunctionSubMenu", "Actor Function"),
-			LOCTEXT("SubsystemFunctionSubMenuToolTip", "Expose an actor's function."),
+			LOCTEXT("ActorFunctionSubMenuToolTip", "Expose an actor's function."),
 			ActorFunctionPicker.ToSharedRef()
 		);
 		
diff --git a/Engine/Plugins/VirtualProduction/RemoteControlWebInterface/Source/RemoteControlWebInterface/Private/RCWebInterfaceProcess.cpp b/Engine/Plugins/VirtualProduction/RemoteControlWebInterface/Source/RemoteControlWebInterface/Private/RCWebInterfaceProcess.cpp
--- a/Engine/Plugins/VirtualProduction/RemoteControlWebInterface/Source/RemoteControlWebInterface/Private/RCWebInterfaceProcess.cpp
+++ b/Engine/Plugins/VirtualProduction/RemoteControlWebInterface/Source/RemoteControlWebInterface/Private/RCWebInterfaceProcess.cpp
@@ -264,7 +264,7 @@ uint32 FRemoteControlWebInterfaceProcess::Run()
 	{
 		TaskNotification->SetComplete(
 			ErrorTitle,
-			LOCTEXT("RemoteControlWebInterface_LaunchFailed", "WebApp exited"),
+			LOCTEXT("RemoteControlWebInterface_AppExited", "WebApp exited"),
 			false
 		);
 	}
diff --git a/Engine/Plugins/VirtualProduction/StageMonitoring/Source/StageMonitorEditor/Private/SStageMonitorPanel.cpp b/Engine/Plugins/VirtualProduction/StageMonitoring/Source/StageMonitorEditor/Private/SStageMonitorPanel.cpp
--- a/Engine/Plugins/VirtualProduction/StageMonitoring/Source/StageMonitorEditor/Private/SStageMonitorPanel.cpp
+++ b/Engine/Plugins/VirtualProduction/StageMonitoring/Source/StageMonitorEditor/Private/SStageMonitorPanel.cpp
@@ -202,7 +202,7 @@ void SStageMonitorPanel::Construct(const FArguments& InArgs)
 					[
 						SNew(SCheckBox)
 						.Style(FStageMonitorEditorStyle::Get(), "ViewMode")
-						.ToolTipText(LOCTEXT("ViewModeTooltip", "Start / stop monitor"))
+						.ToolTipText(LOCTEXT("ActivateMonitorTooltip", "Start / stop monitor"))
 						.IsChecked(this, &SStageMonitorPanel::IsMonitorActive)
 						.OnCheckStateChanged(this, &SStageMonitorPanel::OnMonitorStateChanged)
 						[
diff --git a/Engine/Plugins/VirtualProduction/Takes/Source/TakeRecorderSources/Private/TakeRecorderSourcesModule.cpp b/Engine/Plugins/VirtualProduction/Takes/Source/TakeRecorderSources/Private/TakeRecorderSourcesModule.cpp
--- a/Engine/Plugins/VirtualProduction/Takes/Source/TakeRecorderSources/Private/TakeRecorderSourcesModule.cpp
+++ b/Engine/Plugins/VirtualProduction/Takes/Source/TakeRecorderSources/Private/TakeRecorderSourcesModule.cpp
@@ -455,7 +455,7 @@ public:
 				FText RecordText;
 				if (SelectedActors.Num() == 1)
 				{
-					RecordText = FText::Format(LOCTEXT("RecordSelectedActorsText", "Record {0} with Take Recorder"), FText::FromString(SelectedActors[0]->GetActorLabel()));
+					RecordText = FText::Format(LOCTEXT("RecordSingleSelectedActorText", "Record {0} with Take Recorder"), FText::FromString(SelectedActors[0]->GetActorLabel()));
 				}
 				else
 				{
