Fix some remaining loc duplicate keys

diff --git a/Engine/Plugins/VirtualProduction/CameraCalibration/Source/CameraCalibrationEditor/Private/Calibrators/CameraNodalOffsetAlgoPoints.cpp b/Engine/Plugins/VirtualProduction/CameraCalibration/Source/CameraCalibrationEditor/Private/Calibrators/CameraNodalOffsetAlgoPoints.cpp
--- a/Engine/Plugins/VirtualProduction/CameraCalibration/Source/CameraCalibrationEditor/Private/Calibrators/CameraNodalOffsetAlgoPoints.cpp
+++ b/Engine/Plugins/VirtualProduction/CameraCalibration/Source/CameraCalibrationEditor/Private/Calibrators/CameraNodalOffsetAlgoPoints.cpp
@@ -1346,7 +1346,7 @@ bool UCameraNodalOffsetAlgoPoints::GetNodalOffset(FNodalPointOffset& OutNodalOff
 	if (!AverageSinglePoseResults(SinglePoseResults, AverageTransform))
 	{
 		OutErrorMessage = LOCTEXT("CouldNotAverageSinglePoseResults",
-			"There was an error when trying to average the single pose results");
+			"There was an error when averaging the single pose results");
 
 		return false;
 	}
