Fix PVS static analysis warnings

diff --git a/Engine/Plugins/Runtime/OpenCV/Source/OpenCVHelper/Private/OpenCVBlueprintFunctionLibrary.cpp b/Engine/Plugins/Runtime/OpenCV/Source/OpenCVHelper/Private/OpenCVBlueprintFunctionLibrary.cpp
--- a/Engine/Plugins/Runtime/OpenCV/Source/OpenCVHelper/Private/OpenCVBlueprintFunctionLibrary.cpp
+++ b/Engine/Plugins/Runtime/OpenCV/Source/OpenCVHelper/Private/OpenCVBlueprintFunctionLibrary.cpp
@@ -93,9 +93,13 @@ int32 UOpenCVBlueprintFunctionLibrary::OpenCVChessboardDetectCorners(const UText
 		cv::cornerSubPix(CvFrameGray, Corners, WinSize, ZeroZone, TermCriteria);
 
 		OutDetectedCorners.Reserve(NumCornersFound);
-		for (const cv::Point2f& Corner : Corners)
+
+		if (!Corners.empty())
 		{
-			OutDetectedCorners.Add(FVector2D(Corner.x, Corner.y));
+			for (const cv::Point2f& Corner : Corners)
+			{
+				OutDetectedCorners.Add(FVector2D(Corner.x, Corner.y));
+			}
 		}
 	}
 
diff --git a/Engine/Plugins/VirtualProduction/CameraCalibration/Source/CameraCalibrationEditor/Private/Calibrators/CameraLensDistortionAlgoCheckerboard.cpp b/Engine/Plugins/VirtualProduction/CameraCalibration/Source/CameraCalibrationEditor/Private/Calibrators/CameraLensDistortionAlgoCheckerboard.cpp
--- a/Engine/Plugins/VirtualProduction/CameraCalibration/Source/CameraCalibrationEditor/Private/Calibrators/CameraLensDistortionAlgoCheckerboard.cpp
+++ b/Engine/Plugins/VirtualProduction/CameraCalibration/Source/CameraCalibrationEditor/Private/Calibrators/CameraLensDistortionAlgoCheckerboard.cpp
@@ -280,31 +280,34 @@ bool UCameraLensDistortionAlgoCheckerboard::AddCalibrationRow(FText& OutErrorMes
 		std::vector<cv::Point2f> Corners;
 
 		const bool bCornersFound = cv::findChessboardCorners(
-			CvGray, 
+			CvGray,
 			CheckerboardSize,
 			Corners,
 			cv::CALIB_CB_ADAPTIVE_THRESH | cv::CALIB_CB_NORMALIZE_IMAGE
 		);
-		
+
 		if (!bCornersFound)
 		{
 			OutErrorMessage = FText::FromString(FString::Printf(TEXT(
 				"Could not identify the expected checkerboard points of interest. "
-				"The expected checkerboard has %dx%d inner corners."), 
+				"The expected checkerboard has %dx%d inner corners."),
 				Row->NumCornerCols, Row->NumCornerRows)
 			);
 
 			return false;
 		}
-		
+
 		// cv::TermCriteria::Type::EPS will stop the search when the error is under the given epsilon.
 		// cv::TermCriteria::Type::COUNT will stop after the specified number of iterations regardless of epsilon.
 		cv::TermCriteria Criteria(cv::TermCriteria::Type::EPS | cv::TermCriteria::Type::COUNT, 30, 0.001);
 		cv::cornerSubPix(CvGray, Corners, cv::Size(11, 11), cv::Size(-1, -1), Criteria);
 
-		for (cv::Point2f& Corner : Corners)
+		if (!Corners.empty())
 		{
-			Row->Points2d.Add(FVector2D(Corner.x, Corner.y));
+			for (cv::Point2f& Corner : Corners)
+			{
+				Row->Points2d.Add(FVector2D(Corner.x, Corner.y));
+			}
 		}
 
 		// Update the coverage overlay image with information from the newly added row
diff --git a/Engine/Plugins/VirtualProduction/CameraCalibration/Source/CameraCalibrationEditor/Private/Calibrators/CameraNodalOffsetAlgoCheckerboard.cpp b/Engine/Plugins/VirtualProduction/CameraCalibration/Source/CameraCalibrationEditor/Private/Calibrators/CameraNodalOffsetAlgoCheckerboard.cpp
--- a/Engine/Plugins/VirtualProduction/CameraCalibration/Source/CameraCalibrationEditor/Private/Calibrators/CameraNodalOffsetAlgoCheckerboard.cpp
+++ b/Engine/Plugins/VirtualProduction/CameraCalibration/Source/CameraCalibrationEditor/Private/Calibrators/CameraNodalOffsetAlgoCheckerboard.cpp
@@ -180,9 +180,12 @@ bool UCameraNodalOffsetAlgoCheckerboard::PopulatePoints(FText& OutErrorMessage)
 
 		Points2d.reserve(Corners.size());
 
-		for (cv::Point2f& Corner : Corners)
+		if (!Corners.empty())
 		{
-			Points2d.push_back(cv::Point2f(Corner.x, Corner.y));
+			for (cv::Point2f& Corner : Corners)
+			{
+				Points2d.push_back(cv::Point2f(Corner.x, Corner.y));
+			}
 		}
 	}
 
