[CameraCalibration] Fix NonUnity compile error in CameraLensDistortionAlgoCheckerboard.cpp

diff --git a/Engine/Plugins/VirtualProduction/CameraCalibration/Source/CameraCalibrationEditor/Private/Calibrators/CameraLensDistortionAlgoCheckerboard.cpp b/Engine/Plugins/VirtualProduction/CameraCalibration/Source/CameraCalibrationEditor/Private/Calibrators/CameraLensDistortionAlgoCheckerboard.cpp
--- a/Engine/Plugins/VirtualProduction/CameraCalibration/Source/CameraCalibrationEditor/Private/Calibrators/CameraLensDistortionAlgoCheckerboard.cpp
+++ b/Engine/Plugins/VirtualProduction/CameraCalibration/Source/CameraCalibrationEditor/Private/Calibrators/CameraLensDistortionAlgoCheckerboard.cpp
@@ -7,6 +7,7 @@
 #include "AssetEditor/SSimulcamViewport.h"
 #include "CameraCalibrationCheckerboard.h"
 #include "EditorFontGlyphs.h"
+#include "Engine/Texture2D.h"
 #include "EngineUtils.h"
 #include "Misc/MessageDialog.h"
 #include "PropertyCustomizationHelpers.h"
