Update nodal offset algos to check the lens component to determine whether nodal offset was being applied during a calibration

diff --git a/Engine/Plugins/VirtualProduction/CameraCalibration/Source/CameraCalibrationEditor/Private/Calibrators/CameraNodalOffsetAlgoOpticalAxis.cpp b/Engine/Plugins/VirtualProduction/CameraCalibration/Source/CameraCalibrationEditor/Private/Calibrators/CameraNodalOffsetAlgoOpticalAxis.cpp
--- a/Engine/Plugins/VirtualProduction/CameraCalibration/Source/CameraCalibrationEditor/Private/Calibrators/CameraNodalOffsetAlgoOpticalAxis.cpp
+++ b/Engine/Plugins/VirtualProduction/CameraCalibration/Source/CameraCalibrationEditor/Private/Calibrators/CameraNodalOffsetAlgoOpticalAxis.cpp
@@ -255,7 +255,7 @@ bool UCameraNodalOffsetAlgoOpticalAxis::GetNodalOffset(FNodalPointOffset& OutNod
 	// Cache the existing nodal offset, if one exists
 	CachedNodalOffset = FTransform::Identity;
 
-	if (FirstRow->CameraData.LensFileEvalData.NodalOffset.bWasApplied)
+	if (FirstRow->CameraData.bWasNodalOffsetApplied)
 	{
 		FNodalPointOffset NodalPointOffset;
 
diff --git a/Engine/Plugins/VirtualProduction/CameraCalibration/Source/CameraCalibrationEditor/Private/Calibrators/CameraNodalOffsetAlgoPoints.cpp b/Engine/Plugins/VirtualProduction/CameraCalibration/Source/CameraCalibrationEditor/Private/Calibrators/CameraNodalOffsetAlgoPoints.cpp
--- a/Engine/Plugins/VirtualProduction/CameraCalibration/Source/CameraCalibrationEditor/Private/Calibrators/CameraNodalOffsetAlgoPoints.cpp
+++ b/Engine/Plugins/VirtualProduction/CameraCalibration/Source/CameraCalibrationEditor/Private/Calibrators/CameraNodalOffsetAlgoPoints.cpp
@@ -16,6 +16,7 @@
 #include "Input/Events.h"
 #include "Internationalization/Text.h"
 #include "Layout/Geometry.h"
+#include "LensComponent.h"
 #include "LensDistortionModelHandlerBase.h"
 #include "LensFile.h"
 #include "Math/Vector.h"
@@ -347,7 +348,7 @@ double UCameraNodalOffsetAlgoPoints::MinimizeReprojectionError(FTransform& InOut
 
 	FTransform ExistingOffset = FTransform::Identity;
 
-	if (FirstRow->CameraData.LensFileEvalData.NodalOffset.bWasApplied)
+	if (FirstRow->CameraData.bWasNodalOffsetApplied)
 	{
 		FNodalPointOffset NodalPointOffset;
 		const float Focus = FirstRow->CameraData.LensFileEvalData.Input.Focus;
@@ -483,7 +484,7 @@ double UCameraNodalOffsetAlgoPoints::ComputeReprojectionError(const FTransform&
 
 	FTransform ExistingOffset = FTransform::Identity;
 
-	if (FirstRow->CameraData.LensFileEvalData.NodalOffset.bWasApplied)
+	if (FirstRow->CameraData.bWasNodalOffsetApplied)
 	{
 		FNodalPointOffset NodalPointOffset;
 		const float Focus = FirstRow->CameraData.LensFileEvalData.Input.Focus;
@@ -607,6 +608,7 @@ void UCameraNodalOffsetAlgoPoints::Tick(float DeltaTime)
 			}
 
 			LastCameraData.Pose = CameraComponent->GetComponentToWorld();
+			LastCameraData.bWasNodalOffsetApplied = DoesCameraPoseIncludeNodalOffset(Camera);
 			LastCameraData.UniqueId = Camera->GetUniqueID();
 			LastCameraData.LensFileEvalData = *LensFileEvalData;
 
@@ -913,7 +915,7 @@ bool UCameraNodalOffsetAlgoPoints::ValidateNewRow(TSharedPtr<FCalibrationRowData
 	//
 	// It can't change because we need to know if the camera pose is being affected or not by the current nodal offset evaluation.
 	// And we need to know that because the offset we calculate will need to either subtract or not the current evaluation when adding it to the LUT.
-	if (FirstRow->CameraData.LensFileEvalData.NodalOffset.bWasApplied != Row->CameraData.LensFileEvalData.NodalOffset.bWasApplied)
+	if (FirstRow->CameraData.bWasNodalOffsetApplied != Row->CameraData.bWasNodalOffsetApplied)
 	{
 		OutErrorMessage = LOCTEXT("ApplyNodalOffsetChanged", "Apply nodal offset changed");
 		return false;
@@ -1255,7 +1257,7 @@ bool UCameraNodalOffsetAlgoPoints::GetNodalOffsetSinglePose(
 
 	FTransform ExistingOffset = FTransform::Identity;
 
-	if (FirstRow->CameraData.LensFileEvalData.NodalOffset.bWasApplied)
+	if (FirstRow->CameraData.bWasNodalOffsetApplied)
 	{
 		FNodalPointOffset NodalPointOffset;
 
@@ -1613,6 +1615,38 @@ bool UCameraNodalOffsetAlgoPoints::IsCalibrationComponentSelected(const UCalibra
 	return ActiveCalibratorComponents.Contains(SelectedComponent);
 }
 
+bool UCameraNodalOffsetAlgoPoints::DoesCameraPoseIncludeNodalOffset(const ACameraActor* Camera) const
+{
+	bool bResult = false;
+
+	if (const FCameraCalibrationStepsController* StepsController = NodalOffsetTool->GetCameraCalibrationStepsController())
+	{
+		const ULensFile* LensFile = StepsController->GetLensFile();
+
+		if (LensFile)
+		{
+			TInlineComponentArray<ULensComponent*> LensComponents;
+			Camera->GetComponents(LensComponents);
+
+			for (const ULensComponent* LensComponent : LensComponents)
+			{
+				if (LensComponent->GetLensFile() == LensFile)
+				{
+					bResult = LensComponent->WasNodalOffsetAppliedThisTick();
+				}
+				else if (LensComponent->WasNodalOffsetAppliedThisTick())
+				{
+					UE_LOG(LogCameraCalibrationEditor, VeryVerbose,
+						TEXT("The camera pose includes a nodal offset from a different lens file than the one currently being calibrated." \
+							"If you see this warning, you should disable the application of nodal offset from other lens components before continuing calibration."))
+				}
+			}
+		}
+	}
+
+	return bResult;
+}
+
 bool UCameraNodalOffsetAlgoPoints::ApplyNodalOffsetToCalibrator()
 {
 	using namespace CameraNodalOffsetAlgoPoints;
diff --git a/Engine/Plugins/VirtualProduction/CameraCalibration/Source/CameraCalibrationEditor/Private/Calibrators/CameraNodalOffsetAlgoPoints.h b/Engine/Plugins/VirtualProduction/CameraCalibration/Source/CameraCalibrationEditor/Private/Calibrators/CameraNodalOffsetAlgoPoints.h
--- a/Engine/Plugins/VirtualProduction/CameraCalibration/Source/CameraCalibrationEditor/Private/Calibrators/CameraNodalOffsetAlgoPoints.h
+++ b/Engine/Plugins/VirtualProduction/CameraCalibration/Source/CameraCalibrationEditor/Private/Calibrators/CameraNodalOffsetAlgoPoints.h
@@ -94,6 +94,9 @@ protected:
 		// The camera pose
 		FTransform Pose;
 
+		// True if the cached camera pose included a nodal offset
+		bool bWasNodalOffsetApplied;
+
 		// The data used to evaluate the lens data in the camera for this sample
 		FLensFileEvalData LensFileEvalData;
 
@@ -229,6 +232,9 @@ protected:
 	/** Returns true if the input component is in the set of active calibration components */
 	bool IsCalibrationComponentSelected(const UCalibrationPointComponent* const SelectedComponent) const;
 
+	/** Check whether the input camera has a lens component that is currently applying a nodal offset to the camera pose */
+	bool DoesCameraPoseIncludeNodalOffset(const ACameraActor* Camera) const;
+
 	/** Update any calibration components in the set of active components if they were replaced by an in-editor event */
 	void OnObjectsReplaced(const TMap<UObject*, UObject*>& OldToNewInstanceMap);
 
diff --git a/Engine/Plugins/VirtualProduction/CameraCalibrationCore/Source/CameraCalibrationCore/Private/LensComponent.cpp b/Engine/Plugins/VirtualProduction/CameraCalibrationCore/Source/CameraCalibrationCore/Private/LensComponent.cpp
--- a/Engine/Plugins/VirtualProduction/CameraCalibrationCore/Source/CameraCalibrationCore/Private/LensComponent.cpp
+++ b/Engine/Plugins/VirtualProduction/CameraCalibrationCore/Source/CameraCalibrationCore/Private/LensComponent.cpp
@@ -368,6 +368,8 @@ void ULensComponent::ApplyNodalOffset()
 	TrackedComponent.Get()->AddLocalOffset(Offset.LocationOffset);
 	TrackedComponent.Get()->AddLocalRotation(Offset.RotationOffset);
 
+	bWasNodalOffsetAppliedThisTick = true;
+
 	// Reset so that nodal offset will only be applied again next tick if new tracking data was applied between now and then
 	TrackedComponent.Reset();
 }
@@ -433,6 +435,11 @@ void ULensComponent::SetApplyNodalOffsetOnTick(bool bApplyNodalOffset)
 	bApplyNodalOffsetOnTick = bApplyNodalOffset;
 }
 
+bool ULensComponent::WasNodalOffsetAppliedThisTick() const 
+{
+	return bWasNodalOffsetAppliedThisTick;
+}
+
 void ULensComponent::InitDefaultCamera()
 {
 	UCineCameraComponent* const CineCameraComponent = Cast<UCineCameraComponent>(TargetCameraComponent.GetComponent(GetOwner()));
@@ -550,6 +557,9 @@ void ULensComponent::UpdateTrackedComponent(const ULiveLinkComponentController*
 			if (bIsLocationRotationSupported)
 			{
 				TrackedComponent = Cast<USceneComponent>(TransformController->GetAttachedComponent());
+
+				// At this point, we know that the tracked component's transform was just set by LiveLink, and therefore has no nodal offset applied
+				bWasNodalOffsetAppliedThisTick = false;
 			}
 		}
 	}
diff --git a/Engine/Plugins/VirtualProduction/CameraCalibrationCore/Source/CameraCalibrationCore/Public/LensComponent.h b/Engine/Plugins/VirtualProduction/CameraCalibrationCore/Source/CameraCalibrationCore/Public/LensComponent.h
--- a/Engine/Plugins/VirtualProduction/CameraCalibrationCore/Source/CameraCalibrationCore/Public/LensComponent.h
+++ b/Engine/Plugins/VirtualProduction/CameraCalibrationCore/Source/CameraCalibrationCore/Public/LensComponent.h
@@ -76,6 +76,10 @@ public:
 	UFUNCTION(BlueprintCallable, Category="Lens Component")
 	void SetApplyNodalOffsetOnTick(bool bApplyNodalOffset);
 
+	/** Returns true if nodal offset was applied during the current tick, false otherwise */
+	UFUNCTION(BlueprintPure, Category = "Lens Component")
+	bool WasNodalOffsetAppliedThisTick() const;
+
 	/** 
 	 * Manually apply nodal offset to the specified component. 
 	 * If bUseManualInputs is true, the input Focus and Zoom values will be used to evaluate the LensFile .
@@ -206,4 +210,7 @@ private:
 
 	/** Latest LiveLink FIZ data, used to evaluate the LensFile */
 	FLensFileEvalData LiveLinkFIZ;
+
+	/** Whether or not nodal offset was applied to a tracked component this tick */
+	bool bWasNodalOffsetAppliedThisTick = false;
 };
