Fix bad nodal offset calibration when applying to camera parent, caused by a missing call to UndistortCalibrationRowPoints()

diff --git a/Engine/Plugins/VirtualProduction/CameraCalibration/Source/CameraCalibrationEditor/Private/Calibrators/CameraNodalOffsetAlgoPoints.cpp b/Engine/Plugins/VirtualProduction/CameraCalibration/Source/CameraCalibrationEditor/Private/Calibrators/CameraNodalOffsetAlgoPoints.cpp
--- a/Engine/Plugins/VirtualProduction/CameraCalibration/Source/CameraCalibrationEditor/Private/Calibrators/CameraNodalOffsetAlgoPoints.cpp
+++ b/Engine/Plugins/VirtualProduction/CameraCalibration/Source/CameraCalibrationEditor/Private/Calibrators/CameraNodalOffsetAlgoPoints.cpp
@@ -1773,6 +1773,9 @@ bool UCameraNodalOffsetAlgoPoints::CalcTrackingOriginPoseForSingleCamPose(
 
 	FTransform DesiredCameraPose;
 
+	// Undistort the 2D points in each calibration row 
+	UndistortCalibrationRowPoints();
+
 	if (!CalculatedOptimalCameraComponentPose(DesiredCameraPose, Rows, OutErrorMessage))
 	{
 		return false;
