Update the SimulcamEdior viewport to use the newer version of InputKey and InputAxis.

#rnx

diff --git a/Engine/Plugins/VirtualProduction/CameraCalibration/Source/CameraCalibrationEditor/Private/AssetEditor/Viewport/SimulcamEditorViewportClient.cpp b/Engine/Plugins/VirtualProduction/CameraCalibration/Source/CameraCalibrationEditor/Private/AssetEditor/Viewport/SimulcamEditorViewportClient.cpp
--- a/Engine/Plugins/VirtualProduction/CameraCalibration/Source/CameraCalibrationEditor/Private/AssetEditor/Viewport/SimulcamEditorViewportClient.cpp
+++ b/Engine/Plugins/VirtualProduction/CameraCalibration/Source/CameraCalibrationEditor/Private/AssetEditor/Viewport/SimulcamEditorViewportClient.cpp
@@ -110,19 +110,19 @@ void FSimulcamEditorViewportClient::ZoomTowardsFit(FViewport* Viewport)
 	}
 }
 
-bool FSimulcamEditorViewportClient::InputKey(FViewport* Viewport, int32 ControllerId, FKey Key, EInputEvent Event, float AmountDepressed, bool Gamepad)
+bool FSimulcamEditorViewportClient::InputKey(const FInputKeyEventArgs& InEventArgs)
 {
-	if (Event == IE_Pressed)
+	if (InEventArgs.Event == IE_Pressed)
 	{
 		if (!SimulcamEditorViewportWeakPtr.IsValid())
 		{
 			return false;
 		}
 
-		if (Key == EKeys::LeftMouseButton || Key == EKeys::MiddleMouseButton || Key == EKeys::RightMouseButton)
+		if (InEventArgs.Key == EKeys::LeftMouseButton || InEventArgs.Key == EKeys::MiddleMouseButton || InEventArgs.Key == EKeys::RightMouseButton)
 		{
 			const FGeometry& MyGeometry = SimulcamEditorViewportWeakPtr.Pin()->GetTickSpaceGeometry();
-			const FVector2D LocalMouse = FVector2D(Viewport->GetMouseX(), Viewport->GetMouseY());
+			const FVector2D LocalMouse = FVector2D(InEventArgs.Viewport->GetMouseX(), InEventArgs.Viewport->GetMouseY());
 			// check if we are under the viewport, otherwise the capture system will blindly trigger the PointerEvent
 			if (LocalMouse >= FVector2D(0, 0) && LocalMouse < MyGeometry.GetAbsoluteSize())
 			{
@@ -145,7 +145,7 @@ bool FSimulcamEditorViewportClient::InputKey(FViewport* Viewport, int32 Controll
 								FakeMousePosition,
 								FakeMousePosition,
 								TSet<FKey>(),
-								Key,
+								InEventArgs.Key,
 								0,
 								FSlateApplication::Get().GetPlatformApplication()->GetModifierKeys());
 
@@ -160,42 +160,42 @@ bool FSimulcamEditorViewportClient::InputKey(FViewport* Viewport, int32 Controll
 		{
 			const bool bIsCtrlDown = FSlateApplication::Get().GetPlatformApplication()->GetModifierKeys().IsControlDown();
 
-			if (Key == EKeys::MouseScrollUp || (Key == EKeys::Add && bIsCtrlDown))
+			if (InEventArgs.Key == EKeys::MouseScrollUp || (InEventArgs.Key == EKeys::Add && bIsCtrlDown))
 			{
-				ZoomOnPoint(Viewport, MousePosition, [this] {SimulcamEditorViewportWeakPtr.Pin()->ZoomIn(); });
+				ZoomOnPoint(InEventArgs.Viewport, MousePosition, [this] {SimulcamEditorViewportWeakPtr.Pin()->ZoomIn(); });
 				return true;
 			}
 
-			if (Key == EKeys::MouseScrollDown || (Key == EKeys::Subtract && bIsCtrlDown))
+			if (InEventArgs.Key == EKeys::MouseScrollDown || (InEventArgs.Key == EKeys::Subtract && bIsCtrlDown))
 			{	
 				FVector2D TextureSize = SimulcamEditorViewportWeakPtr.Pin()->CalculateTextureDimensions();
-				if (!bIsCtrlDown && (CurrentTexturePosition.X > 0 || (TextureSize.X + CurrentTexturePosition.X) < Viewport->GetSizeXY().X || CurrentTexturePosition.Y > 0 || (TextureSize.Y + CurrentTexturePosition.Y) < Viewport->GetSizeXY().Y))
+				if (!bIsCtrlDown && (CurrentTexturePosition.X > 0 || (TextureSize.X + CurrentTexturePosition.X) < InEventArgs.Viewport->GetSizeXY().X || CurrentTexturePosition.Y > 0 || (TextureSize.Y + CurrentTexturePosition.Y) < InEventArgs.Viewport->GetSizeXY().Y))
 				{
-					ZoomTowardsFit(Viewport);
+					ZoomTowardsFit(InEventArgs.Viewport);
 				}
 				else
 				{
-					ZoomOnPoint(Viewport, MousePosition, [this] {SimulcamEditorViewportWeakPtr.Pin()->ZoomOut(); });
+					ZoomOnPoint(InEventArgs.Viewport, MousePosition, [this] {SimulcamEditorViewportWeakPtr.Pin()->ZoomOut(); });
 				}
 				return true;
 			}
 
-			if ((Key == EKeys::Zero || Key == EKeys::NumPadZero) && bIsCtrlDown)
+			if ((InEventArgs.Key == EKeys::Zero || InEventArgs.Key == EKeys::NumPadZero) && bIsCtrlDown)
 			{
-				ZoomToFit(Viewport);
+				ZoomToFit(InEventArgs.Viewport);
 				return true;
 			}
 		}
 
-		return SimulcamViewportWeakPtr.Pin()->OnViewportInputKey(Key, Event);
+		return SimulcamViewportWeakPtr.Pin()->OnViewportInputKey(InEventArgs.Key, InEventArgs.Event);
 	}
-	else if (Event == IE_Released)
+	else if (InEventArgs.Event == IE_Released)
 	{
-		return SimulcamViewportWeakPtr.Pin()->OnViewportInputKey(Key, Event);
+		return SimulcamViewportWeakPtr.Pin()->OnViewportInputKey(InEventArgs.Key, InEventArgs.Event);
 	}
-	else if (Event == IE_Repeat)
+	else if (InEventArgs.Event == IE_Repeat)
 	{
-		return SimulcamViewportWeakPtr.Pin()->OnViewportInputKey(Key, Event);
+		return SimulcamViewportWeakPtr.Pin()->OnViewportInputKey(InEventArgs.Key, InEventArgs.Event);
 	}
 
 	return false;
@@ -261,7 +261,7 @@ bool FSimulcamEditorViewportClient::InputChar(FViewport* Viewport, int32 Control
 	return false;
 }
 
-bool FSimulcamEditorViewportClient::InputAxis(FViewport* Viewport, int32 ControllerId, FKey Key, float Delta, float DeltaTime, int32 NumSamples, bool bGamepad)
+bool FSimulcamEditorViewportClient::InputAxis(FViewport* Viewport, FInputDeviceId DeviceID, FKey Key, float Delta, float DeltaTime, int32 NumSamples, bool bGamepad)
 {
 	if (Key == EKeys::MouseX || Key == EKeys::MouseY)
 	{
diff --git a/Engine/Plugins/VirtualProduction/CameraCalibration/Source/CameraCalibrationEditor/Private/AssetEditor/Viewport/SimulcamEditorViewportClient.h b/Engine/Plugins/VirtualProduction/CameraCalibration/Source/CameraCalibrationEditor/Private/AssetEditor/Viewport/SimulcamEditorViewportClient.h
--- a/Engine/Plugins/VirtualProduction/CameraCalibration/Source/CameraCalibrationEditor/Private/AssetEditor/Viewport/SimulcamEditorViewportClient.h
+++ b/Engine/Plugins/VirtualProduction/CameraCalibration/Source/CameraCalibrationEditor/Private/AssetEditor/Viewport/SimulcamEditorViewportClient.h
@@ -22,9 +22,9 @@ public:
 
 	/** Begin FViewportClient interface */
 	virtual void Draw(FViewport* Viewport, FCanvas* Canvas) override;
-	virtual bool InputKey(FViewport* Viewport, int32 ControllerId, FKey Key, EInputEvent Event, float AmountDepressed = 1.0f, bool bGamepad = false) override;
+	virtual bool InputKey(const FInputKeyEventArgs& InEventArgs) override;
 	virtual bool InputChar(FViewport* Viewport, int32 ControllerId, TCHAR Character) override;
-	virtual bool InputAxis(FViewport* Viewport, int32 ControllerId, FKey Key, float Delta, float DeltaTime, int32 NumSamples=1, bool bGamepad=false) override;
+	virtual bool InputAxis(FViewport* Viewport, FInputDeviceId DeviceId, FKey Key, float Delta, float DeltaTime, int32 NumSamples=1, bool bGamepad=false) override;
 	virtual UWorld* GetWorld() const override { return nullptr; }
 	virtual EMouseCursor::Type GetCursor(FViewport* Viewport, int32 X, int32 Y) override;
 	virtual void MouseMove(FViewport* Viewport, int32 X, int32 Y) override;
