- Fixing lens file import doing wrong coordinate conversion (unreal to opencv instead of the opposite)

diff --git a/Engine/Plugins/VirtualProduction/CameraCalibration/Source/CameraCalibrationEditor/Private/LensFileExchangeFormat.cpp b/Engine/Plugins/VirtualProduction/CameraCalibration/Source/CameraCalibrationEditor/Private/LensFileExchangeFormat.cpp
--- a/Engine/Plugins/VirtualProduction/CameraCalibration/Source/CameraCalibrationEditor/Private/LensFileExchangeFormat.cpp
+++ b/Engine/Plugins/VirtualProduction/CameraCalibration/Source/CameraCalibrationEditor/Private/LensFileExchangeFormat.cpp
@@ -822,7 +822,7 @@ void FLensFileExchange::PopulateNodalOffsetTable(ULensFile& OutLensFile, const F
 
 			if (Metadata.NodalOffsetCoordinateSystem == ENodalOffsetCoordinateSystem::OpenCV)
 			{
-				FOpenCVHelper::ConvertUnrealToOpenCV(NodalOffsetTransform);
+				FOpenCVHelper::ConvertOpenCVToUnreal(NodalOffsetTransform);
 			}
 
 			FNodalPointOffset NodalPointOffset;
