[CameraCalibration] Add "display unit" setting and use it in the "add point" dialog to let users enter focal length and image center in mm, pixels, or normalized.

diff --git a/Engine/Plugins/VirtualProduction/CameraCalibration/Source/CameraCalibrationEditor/Private/AssetEditor/CameraCalibrationToolkit.cpp b/Engine/Plugins/VirtualProduction/CameraCalibration/Source/CameraCalibrationEditor/Private/AssetEditor/CameraCalibrationToolkit.cpp
--- a/Engine/Plugins/VirtualProduction/CameraCalibration/Source/CameraCalibrationEditor/Private/AssetEditor/CameraCalibrationToolkit.cpp
+++ b/Engine/Plugins/VirtualProduction/CameraCalibration/Source/CameraCalibrationEditor/Private/AssetEditor/CameraCalibrationToolkit.cpp
@@ -114,7 +114,9 @@ TSharedPtr<SWindow> FCameraCalibrationToolkit::OpenPopupWindow(const FText& InTi
 			.HasCloseButton(true)
 			.SupportsMaximize(false)
 			.SupportsMinimize(false)
-			.ClientSize(FVector2D(480, 360));
+			.ClientSize(FVector2D(480, 720))
+			.MinWidth(480)
+			.MinHeight(720);
 
 		FSlateApplication::Get().AddWindow(PopupWindowPin.ToSharedRef());
 	}
diff --git a/Engine/Plugins/VirtualProduction/CameraCalibration/Source/CameraCalibrationEditor/Private/AssetEditor/SLensDataAddPointDialog.cpp b/Engine/Plugins/VirtualProduction/CameraCalibration/Source/CameraCalibrationEditor/Private/AssetEditor/SLensDataAddPointDialog.cpp
--- a/Engine/Plugins/VirtualProduction/CameraCalibration/Source/CameraCalibrationEditor/Private/AssetEditor/SLensDataAddPointDialog.cpp
+++ b/Engine/Plugins/VirtualProduction/CameraCalibration/Source/CameraCalibrationEditor/Private/AssetEditor/SLensDataAddPointDialog.cpp
@@ -2,9 +2,9 @@
 
 #include "SLensDataAddPointDialog.h"
 
+#include "CameraCalibrationSettings.h"
 #include "CameraCalibrationToolkit.h"
 #include "Curves/RichCurve.h"
-#include "Styling/AppStyle.h"
 #include "Framework/MultiBox/MultiBoxBuilder.h"
 #include "IDetailChildrenBuilder.h"
 #include "IDetailPropertyRow.h"
@@ -15,6 +15,7 @@
 #include "Modules/ModuleManager.h"
 #include "PropertyHandle.h"
 #include "ScopedTransaction.h"
+#include "Styling/AppStyle.h"
 #include "Widgets/Input/SButton.h"
 #include "Widgets/Input/SCheckBox.h"
 #include "Widgets/Input/SNumericEntryBox.h"
@@ -316,6 +317,54 @@ TSharedRef<SWidget> SLensDataAddPointDialog::MakeTrackingDataWidget()
 	return TrackingWidget.ToSharedRef();
 }
 
+void SLensDataAddPointDialog::OnUnitsChanged()
+{
+	// Refresh the widget with updated default values
+	LensDataContainer->SetContent(MakeLensDataWidget());
+}
+
+FVector2D SLensDataAddPointDialog::GetDefaultFocalLengthValue() const
+{
+	const ELensDisplayUnit DisplayUnit = GetDefault<UCameraCalibrationEditorSettings>()->DefaultDisplayUnit;
+
+	FVector2D DefaultValue;
+	switch (DisplayUnit)
+	{
+	case ELensDisplayUnit::Millimeters:
+		DefaultValue = FVector2D(LensFile->LensInfo.SensorDimensions[0], LensFile->LensInfo.SensorDimensions[0]);
+		break;
+	case ELensDisplayUnit::Pixels:
+		DefaultValue = FVector2D(LensFile->LensInfo.ImageDimensions[0], LensFile->LensInfo.ImageDimensions[0]);
+		break;
+	case ELensDisplayUnit::Normalized: // falls through
+	default:
+		DefaultValue = FVector2D(1.0, (16.0 / 9.0));
+		break;
+	}
+	return DefaultValue;
+}
+
+FVector2D SLensDataAddPointDialog::GetDefaultImageCenterValue() const
+{
+	const ELensDisplayUnit DisplayUnit = GetDefault<UCameraCalibrationEditorSettings>()->DefaultDisplayUnit;
+
+	FVector2D DefaultValue;
+	switch (DisplayUnit)
+	{
+	case ELensDisplayUnit::Millimeters:
+		DefaultValue = FVector2D(LensFile->LensInfo.SensorDimensions[0] * 0.5, LensFile->LensInfo.SensorDimensions[1] * 0.5);
+		break;
+	case ELensDisplayUnit::Pixels:
+		DefaultValue = FVector2D(LensFile->LensInfo.ImageDimensions[0] * 0.5, LensFile->LensInfo.ImageDimensions[1] * 0.5);
+		break;
+	case ELensDisplayUnit::Normalized: // falls through
+	default:
+		DefaultValue = FVector2D(0.5, 0.5);
+		break;
+	}
+	return DefaultValue;
+}
+
 TSharedRef<SWidget> SLensDataAddPointDialog::MakeLensDataWidget()
 {
 	FStructureDetailsViewArgs StructureViewArgs;
@@ -326,6 +375,16 @@ TSharedRef<SWidget> SLensDataAddPointDialog::MakeLensDataWidget()
 	FPropertyEditorModule& PropertyEditor = FModuleManager::Get().LoadModuleChecked<FPropertyEditorModule>(TEXT("PropertyEditor"));
 	TSharedPtr<SWidget> LensDataWidget = SNullWidget::NullWidget;
 	
+	FSinglePropertyParams InitParams;
+	InitParams.NamePlacement = EPropertyNamePlacement::Hidden;
+	const TSharedPtr<ISinglePropertyView> UnitWidget = PropertyEditor.CreateSingleProperty(
+		GetMutableDefault<UCameraCalibrationEditorSettings>(), 
+		GET_MEMBER_NAME_CHECKED(UCameraCalibrationEditorSettings, DefaultDisplayUnit),
+		InitParams);
+
+	FSimpleDelegate OnUnitsChangedDelegate = FSimpleDelegate::CreateSP(this, &SLensDataAddPointDialog::OnUnitsChanged);
+	UnitWidget->SetOnPropertyValueChanged(OnUnitsChangedDelegate);
+
 	switch (SelectedCategory)
 	{
 		case ELensDataCategory::Focus:
@@ -337,8 +396,37 @@ TSharedRef<SWidget> SLensDataAddPointDialog::MakeLensDataWidget()
 		case ELensDataCategory::Zoom:
 		{
 			FocalLengthData->InitializeAs<FFocalLengthInfo>();
+			FFocalLengthInfo* FocalLengthInstanceData = FocalLengthData->Cast<FFocalLengthInfo>();
+			FocalLengthInstanceData->FxFy = GetDefaultFocalLengthValue();
+
 			TSharedPtr<IStructureDetailsView> FocalLengthStructDetailsView = PropertyEditor.CreateStructureDetailView(DetailArgs, StructureViewArgs, FocalLengthData);
-			LensDataWidget = FocalLengthStructDetailsView->GetWidget();
+
+			LensDataWidget =
+				SNew(SVerticalBox)
+				+ SVerticalBox::Slot()
+				.AutoHeight()
+				.Padding(5.0f, 5.0f)
+				[
+					SNew(STextBlock)
+					.Text(LOCTEXT("FxFyDisplayUnitHelpText", 
+						"Choose the units that should be used for the value of FxFy below.\n"
+						"If using mm or pixels, check that the sensor and image dimensions are set correctly in the Lens Information."
+					))
+					.AutoWrapText(true)
+				]
+				+ SVerticalBox::Slot()
+				.AutoHeight()
+				.Padding(0.0f, 5.0f)
+				[
+					UnitWidget.ToSharedRef()
+				]
+				+ SVerticalBox::Slot()
+				.AutoHeight()
+				.Padding(0.0f, 5.0f)
+				[
+					FocalLengthStructDetailsView->GetWidget().ToSharedRef()
+				];
+
 			break;
 		}
 		case ELensDataCategory::Distortion:
@@ -352,14 +440,37 @@ TSharedRef<SWidget> SLensDataAddPointDialog::MakeLensDataWidget()
 				
 			FocalLengthData->InitializeAs<FFocalLengthInfo>();
 			TSharedPtr<IStructureDetailsView> FocalLengthStructDetailsView = PropertyEditor.CreateStructureDetailView(DetailArgs, StructureViewArgs, FocalLengthData);
+			FFocalLengthInfo* FocalLengthInstanceData = FocalLengthData->Cast<FFocalLengthInfo>();
+			FocalLengthInstanceData->FxFy = GetDefaultFocalLengthValue();
 
 			LensDataWidget =
 				SNew(SVerticalBox)
 				+ SVerticalBox::Slot()
+				.AutoHeight()
+				.Padding(5.0f, 5.0f)
+				[
+					SNew(STextBlock)
+					.Text(LOCTEXT("FxFyDisplayUnitHelpText",
+						"Choose the units that should be used for the value of FxFy below.\n"
+						"If using mm or pixels, check that the sensor and image dimensions are set correctly in the Lens Information."
+					))
+					.AutoWrapText(true)
+				]
+				+ SVerticalBox::Slot()
+				.AutoHeight()
+				.Padding(0.0f, 5.0f)
+				[
+					UnitWidget.ToSharedRef()
+				]
+				+ SVerticalBox::Slot()
+				.AutoHeight()
+				.Padding(0.0f, 5.0f)
 				[
 					DistortionDataStructDetailsView->GetWidget().ToSharedRef()
 				]
 				+ SVerticalBox::Slot()
+				.AutoHeight()
+				.Padding(0.0f, 5.0f)
 				[
 					FocalLengthStructDetailsView->GetWidget().ToSharedRef()
 				];
@@ -370,7 +481,35 @@ TSharedRef<SWidget> SLensDataAddPointDialog::MakeLensDataWidget()
 		{
 			ImageCenterData->InitializeAs<FImageCenterInfo>();
 			TSharedPtr<IStructureDetailsView> StructureDetailsView = PropertyEditor.CreateStructureDetailView(DetailArgs, StructureViewArgs, ImageCenterData);
-			LensDataWidget = StructureDetailsView->GetWidget();
+			FImageCenterInfo* ImageCenterInstanceData = ImageCenterData->Cast<FImageCenterInfo>();
+			ImageCenterInstanceData->PrincipalPoint = GetDefaultImageCenterValue();
+
+			LensDataWidget =
+				SNew(SVerticalBox)
+				+ SVerticalBox::Slot()
+				.AutoHeight()
+				.Padding(5.0f, 5.0f)
+				[
+					SNew(STextBlock)
+					.Text(LOCTEXT("ImageCenterDisplayUnitHelpText",
+						"Choose the units that should be used for the value of ImageCenter below.\n"
+						"If using mm or pixels, check that the sensor and image dimensions are set correctly in the Lens Information."
+					))
+				.AutoWrapText(true)
+				]
+				+ SVerticalBox::Slot()
+				.AutoHeight()
+				.Padding(0.0f, 5.0f)
+				[
+					UnitWidget.ToSharedRef()
+				]
+				+ SVerticalBox::Slot()
+				.AutoHeight()
+				.Padding(0.0f, 5.0f)
+				[
+					StructureDetailsView->GetWidget().ToSharedRef()
+				];
+
 			break;
 		}
 		case ELensDataCategory::NodalOffset:
@@ -603,6 +742,39 @@ void SLensDataAddPointDialog::RefreshEvaluationData()
 	}
 }
 
+void SLensDataAddPointDialog::NormalizeValue(FVector2D& InOutValue) const
+{
+	const ELensDisplayUnit DisplayUnit = GetDefault<UCameraCalibrationEditorSettings>()->DefaultDisplayUnit;
+
+	switch (DisplayUnit)
+	{
+	case ELensDisplayUnit::Millimeters:
+	{
+		FVector2D SensorDimensions = LensFile->LensInfo.SensorDimensions;
+		if (SensorDimensions[0] > 0.0 && SensorDimensions[1] > 0.0)
+		{
+			InOutValue[0] = InOutValue[0] / SensorDimensions[0];
+			InOutValue[1] = InOutValue[1] / SensorDimensions[1];
+		}
+		break;
+	}
+	case ELensDisplayUnit::Pixels:
+	{
+		FVector2D ImageDimensions = LensFile->LensInfo.ImageDimensions;
+		if (ImageDimensions[0] > 0.0 && ImageDimensions[1] > 0.0)
+		{
+			InOutValue[0] = InOutValue[0] / ImageDimensions[0];
+			InOutValue[1] = InOutValue[1] / ImageDimensions[1];
+		}
+		break;
+	}
+	case ELensDisplayUnit::Normalized: // falls through
+	default:
+		// do nothing
+		break;
+	}
+}
+
 void SLensDataAddPointDialog::AddDataToLensFile() const
 {
 	switch (SelectedCategory)
@@ -619,17 +791,23 @@ void SLensDataAddPointDialog::AddDataToLensFile() const
 		}
 		case ELensDataCategory::Zoom:
 		{
-			LensFile->AddFocalLengthPoint(TrackingInputData[0].Value, TrackingInputData[1].Value, *FocalLengthData->Get());
+			FFocalLengthInfo FocalLengthValue = *FocalLengthData->Get();
+			NormalizeValue(FocalLengthValue.FxFy);
+			LensFile->AddFocalLengthPoint(TrackingInputData[0].Value, TrackingInputData[1].Value, FocalLengthValue);
 			break;
 		}	
 		case ELensDataCategory::Distortion:
 		{
-			LensFile->AddDistortionPoint(TrackingInputData[0].Value, TrackingInputData[1].Value, DistortionInfoData->Get()->DistortionInfo, *FocalLengthData->Get());
+			FFocalLengthInfo FocalLengthValue = *FocalLengthData->Get();
+			NormalizeValue(FocalLengthValue.FxFy);
+			LensFile->AddDistortionPoint(TrackingInputData[0].Value, TrackingInputData[1].Value, DistortionInfoData->Get()->DistortionInfo, FocalLengthValue);
 			break;
 		}	
 		case ELensDataCategory::ImageCenter:
 		{
-			LensFile->AddImageCenterPoint(TrackingInputData[0].Value, TrackingInputData[1].Value, *ImageCenterData->Get());
+			FImageCenterInfo ImageCenterValue = *ImageCenterData->Get();
+			NormalizeValue(ImageCenterValue.PrincipalPoint);
+			LensFile->AddImageCenterPoint(TrackingInputData[0].Value, TrackingInputData[1].Value, ImageCenterValue);
 			break;
 		}
 		case ELensDataCategory::NodalOffset:
diff --git a/Engine/Plugins/VirtualProduction/CameraCalibration/Source/CameraCalibrationEditor/Private/AssetEditor/SLensDataAddPointDialog.h b/Engine/Plugins/VirtualProduction/CameraCalibration/Source/CameraCalibrationEditor/Private/AssetEditor/SLensDataAddPointDialog.h
--- a/Engine/Plugins/VirtualProduction/CameraCalibration/Source/CameraCalibrationEditor/Private/AssetEditor/SLensDataAddPointDialog.h
+++ b/Engine/Plugins/VirtualProduction/CameraCalibration/Source/CameraCalibrationEditor/Private/AssetEditor/SLensDataAddPointDialog.h
@@ -73,6 +73,18 @@ private:
 	/** Called when ready to add data to LensFile */
 	void AddDataToLensFile() const;
 
+	/** Called when the display units are changed */
+	void OnUnitsChanged();
+
+	/** Returns the default value to display for focal length based on the display units of the LensFile */
+	FVector2D GetDefaultFocalLengthValue() const;
+
+	/** Returns the default value to display for image center based on the display units of the LensFile */
+	FVector2D GetDefaultImageCenterValue() const;
+
+	/** Normalize the input value based on the display units of the LensFile */
+	void NormalizeValue(FVector2D& InputValue) const;
+
 private:
 
 	/**
diff --git a/Engine/Plugins/VirtualProduction/CameraCalibrationCore/Source/CameraCalibrationCore/Public/CameraCalibrationSettings.h b/Engine/Plugins/VirtualProduction/CameraCalibrationCore/Source/CameraCalibrationCore/Public/CameraCalibrationSettings.h
--- a/Engine/Plugins/VirtualProduction/CameraCalibrationCore/Source/CameraCalibrationCore/Public/CameraCalibrationSettings.h
+++ b/Engine/Plugins/VirtualProduction/CameraCalibrationCore/Source/CameraCalibrationCore/Public/CameraCalibrationSettings.h
@@ -175,6 +175,15 @@ struct FLensDataCategoryEditorColor
 	FColor NodalOffset = FColor::Purple;
 };
 
+/** Units used to display/interpret Focal Length and Image Center */
+UENUM()
+enum class ELensDisplayUnit : uint8
+{
+	Millimeters,
+	Pixels,
+	Normalized
+};
+
 /**
  * Settings for the camera calibration when in editor and standalone.
  * @note Cooked games don't use this setting.
@@ -217,6 +226,12 @@ public:
 	UPROPERTY(config, EditAnywhere, Category = "Settings", Meta = (ToolTip = "Enable or Disable Time input driven by Live Link."))
 	bool bEnableTimeSlider = true;
 
+	/** 
+	 * Units used to display/interpret Focal Length and Image Center 
+	 */
+	UPROPERTY(config, EditAnywhere, Category = "Settings")
+	ELensDisplayUnit DefaultDisplayUnit = ELensDisplayUnit::Millimeters;
+
 private:
 
 	/** 
