Remove unacceptable word

diff --git a/Engine/Plugins/VirtualProduction/CameraCalibration/Source/CameraCalibrationEditor/Private/AssetEditor/CameraCalibrationStepsController.h b/Engine/Plugins/VirtualProduction/CameraCalibration/Source/CameraCalibrationEditor/Private/AssetEditor/CameraCalibrationStepsController.h
--- a/Engine/Plugins/VirtualProduction/CameraCalibration/Source/CameraCalibrationEditor/Private/AssetEditor/CameraCalibrationStepsController.h
+++ b/Engine/Plugins/VirtualProduction/CameraCalibration/Source/CameraCalibrationEditor/Private/AssetEditor/CameraCalibrationStepsController.h
@@ -180,7 +180,7 @@ private:
 	/** Finds and creates the available calibration steps */
 	void CreateSteps();
 
-	/** Create a new material transform pass to represent an overlay and add it to the master comp */
+	/** Create a new material transform pass to represent an overlay and add it to the comp */
 	void CreateOverlayPass(FName PassName, TWeakObjectPtr<UCompositingElementMaterialPass>& OverlayPass, TWeakObjectPtr<UTextureRenderTarget2D>& OverlayRenderTarget);
 
 	/** Returns the overlay material pass associated with the input overlay pass type */
