[CameraCalibration] Fix bug with overlay rendering when overlay materials get GC'd

diff --git a/Engine/Plugins/VirtualProduction/CameraCalibration/Source/CameraCalibrationEditor/Private/AssetEditor/ImageCenterTool.h b/Engine/Plugins/VirtualProduction/CameraCalibration/Source/CameraCalibrationEditor/Private/AssetEditor/ImageCenterTool.h
--- a/Engine/Plugins/VirtualProduction/CameraCalibration/Source/CameraCalibrationEditor/Private/AssetEditor/ImageCenterTool.h
+++ b/Engine/Plugins/VirtualProduction/CameraCalibration/Source/CameraCalibrationEditor/Private/AssetEditor/ImageCenterTool.h
@@ -59,7 +59,7 @@ private:
 
 	/** Map of algo names to overlay MIDs used by those algos */
 	UPROPERTY(Transient)
-	TMap<FName, TWeakObjectPtr<UMaterialInstanceDynamic>> AlgoOverlayMIDs;
+	TMap<FName, TObjectPtr<UMaterialInstanceDynamic>> AlgoOverlayMIDs;
 
 	/** True if this tool is the active one in the panel */
 	bool bIsActive = false;
diff --git a/Engine/Plugins/VirtualProduction/CameraCalibration/Source/CameraCalibrationEditor/Private/AssetEditor/LensDistortionTool.h b/Engine/Plugins/VirtualProduction/CameraCalibration/Source/CameraCalibrationEditor/Private/AssetEditor/LensDistortionTool.h
--- a/Engine/Plugins/VirtualProduction/CameraCalibration/Source/CameraCalibrationEditor/Private/AssetEditor/LensDistortionTool.h
+++ b/Engine/Plugins/VirtualProduction/CameraCalibration/Source/CameraCalibrationEditor/Private/AssetEditor/LensDistortionTool.h
@@ -72,7 +72,7 @@ private:
 
 	/** Map of algo names to overlay MIDs used by those algos */
 	UPROPERTY(Transient)
-	TMap<FName, TWeakObjectPtr<UMaterialInstanceDynamic>> AlgoOverlayMIDs;
+	TMap<FName, TObjectPtr<UMaterialInstanceDynamic>> AlgoOverlayMIDs;
 
 	/** True if this tool is the active one in the panel */
 	bool bIsActive = false;
diff --git a/Engine/Plugins/VirtualProduction/CameraCalibration/Source/CameraCalibrationEditor/Private/AssetEditor/NodalOffsetTool.h b/Engine/Plugins/VirtualProduction/CameraCalibration/Source/CameraCalibrationEditor/Private/AssetEditor/NodalOffsetTool.h
--- a/Engine/Plugins/VirtualProduction/CameraCalibration/Source/CameraCalibrationEditor/Private/AssetEditor/NodalOffsetTool.h
+++ b/Engine/Plugins/VirtualProduction/CameraCalibration/Source/CameraCalibrationEditor/Private/AssetEditor/NodalOffsetTool.h
@@ -69,7 +69,7 @@ private:
 
 	/** Map of algo names to overlay MIDs used by those algos */
 	UPROPERTY(Transient)
-	TMap<FName, TWeakObjectPtr<UMaterialInstanceDynamic>> AlgoOverlayMIDs;
+	TMap<FName, TObjectPtr<UMaterialInstanceDynamic>> AlgoOverlayMIDs;
 
 	/** True if this tool is the active one in the panel */
 	bool bIsActive = false;
