[CameraCalibration] Add media playback button controls to lens file editor to support play/pause and stepping in file media sources

diff --git a/Engine/Plugins/VirtualProduction/CameraCalibration/Source/CameraCalibrationEditor/Private/AssetEditor/CameraCalibrationCommands.cpp b/Engine/Plugins/VirtualProduction/CameraCalibration/Source/CameraCalibrationEditor/Private/AssetEditor/CameraCalibrationCommands.cpp
--- a/Engine/Plugins/VirtualProduction/CameraCalibration/Source/CameraCalibrationEditor/Private/AssetEditor/CameraCalibrationCommands.cpp
+++ b/Engine/Plugins/VirtualProduction/CameraCalibration/Source/CameraCalibrationEditor/Private/AssetEditor/CameraCalibrationCommands.cpp
@@ -13,6 +13,7 @@
 void FCameraCalibrationCommands::RegisterCommands()
 {
 	UI_COMMAND(Edit, "Edit", "Edit the current lens file.", EUserInterfaceActionType::Button, FInputChord());
+	UI_COMMAND(ShowMediaPlaybackControls, "Show Media Playback Controls", "Show the media playback control buttons", EUserInterfaceActionType::ToggleButton, FInputChord());
 }
 
 #undef LOCTEXT_NAMESPACE
\ No newline at end of file
diff --git a/Engine/Plugins/VirtualProduction/CameraCalibration/Source/CameraCalibrationEditor/Private/AssetEditor/CameraCalibrationCommands.h b/Engine/Plugins/VirtualProduction/CameraCalibration/Source/CameraCalibrationEditor/Private/AssetEditor/CameraCalibrationCommands.h
--- a/Engine/Plugins/VirtualProduction/CameraCalibration/Source/CameraCalibrationEditor/Private/AssetEditor/CameraCalibrationCommands.h
+++ b/Engine/Plugins/VirtualProduction/CameraCalibration/Source/CameraCalibrationEditor/Private/AssetEditor/CameraCalibrationCommands.h
@@ -25,4 +25,8 @@ public:
 
 	/** Edit the current lens distortion. */
 	TSharedPtr<FUICommandInfo> Edit;
+
+	/** Show the media playback control buttons. */
+	TSharedPtr<FUICommandInfo> ShowMediaPlaybackControls;
+
 };
diff --git a/Engine/Plugins/VirtualProduction/CameraCalibration/Source/CameraCalibrationEditor/Private/AssetEditor/CameraCalibrationStepsController.cpp b/Engine/Plugins/VirtualProduction/CameraCalibration/Source/CameraCalibrationEditor/Private/AssetEditor/CameraCalibrationStepsController.cpp
--- a/Engine/Plugins/VirtualProduction/CameraCalibration/Source/CameraCalibrationEditor/Private/AssetEditor/CameraCalibrationStepsController.cpp
+++ b/Engine/Plugins/VirtualProduction/CameraCalibration/Source/CameraCalibrationEditor/Private/AssetEditor/CameraCalibrationStepsController.cpp
@@ -6,6 +6,7 @@
 #include "CalibrationPointComponent.h"
 #include "Camera/CameraActor.h"
 #include "CameraCalibrationEditorLog.h"
+#include "CameraCalibrationSettings.h"
 #include "CameraCalibrationStep.h"
 #include "CameraCalibrationSubsystem.h"
 #include "CameraCalibrationToolkit.h"
@@ -473,7 +474,7 @@ void FCameraCalibrationStepsController::CreateComp()
 		}
 
 		MediaPlayer->PlayOnOpen = true;
-
+		MediaPlayer->SetLooping(true);
 		// Create MediaTexture
 
 		MediaTexture = NewObject<UMediaTexture>(GetTransientPackage(), NAME_None, RF_Transient);
@@ -894,6 +895,164 @@ bool FCameraCalibrationStepsController::OnSimulcamViewportInputKey(const FKey& I
 	return bStepHandled;
 }
 
+FReply FCameraCalibrationStepsController::OnRewindButtonClicked()
+{
+	// Rewind to the beginning of the media
+	MediaPlayer->Rewind();
+	return FReply::Handled();
+}
+
+FReply FCameraCalibrationStepsController::OnReverseButtonClicked()
+{
+	// Increase the reverse media playback rate
+	MediaPlayer->SetRate(GetFasterReverseRate());
+	return FReply::Handled();
+}
+
+FReply FCameraCalibrationStepsController::OnStepBackButtonClicked() 
+{
+	const float DefaultStepRateInMilliseconds = GetDefault<UCameraCalibrationEditorSettings>()->DefaultMediaStepRateInMilliseconds;
+	const bool bForceDefaultStepRate = GetDefault<UCameraCalibrationEditorSettings>()->bForceDefaultMediaStepRate;
+
+	// The media player could return a frame rate of 0 for the current video track
+	const float MediaFrameRate = MediaPlayer->GetVideoTrackFrameRate(INDEX_NONE, INDEX_NONE);
+
+	// Use the default step rate if the media player returned an invalid frame rate or if the project settings force it
+	float MillisecondsPerStep = 0.0f;
+	if (FMath::IsNearlyEqual(MediaFrameRate, 0.0f) || bForceDefaultStepRate)
+	{
+		MillisecondsPerStep = DefaultStepRateInMilliseconds;
+	}
+	else
+	{
+		MillisecondsPerStep = 1000.0f / MediaFrameRate;
+	}
+
+	// Compute the number of ticks in one step and go backward from the media's current time (clamping to 0)
+	const FTimespan TicksInOneStep = ETimespan::TicksPerMillisecond * (MillisecondsPerStep);
+	FTimespan PreviousStepTime = MediaPlayer->GetTime() - TicksInOneStep;
+	if (PreviousStepTime < FTimespan::Zero())
+	{
+		PreviousStepTime = FTimespan::Zero();
+	}
+
+	MediaPlayer->Seek(PreviousStepTime);
+	MediaPlayer->Pause();
+
+	return FReply::Handled();
+}
+
+FReply FCameraCalibrationStepsController::OnPlayButtonClicked() 
+{
+	MediaPlayer->Play();
+	return FReply::Handled(); 
+}
+
+FReply FCameraCalibrationStepsController::OnPauseButtonClicked() 
+{
+	MediaPlayer->Pause();
+	return FReply::Handled(); 
+}
+
+FReply FCameraCalibrationStepsController::OnStepForwardButtonClicked()
+{
+	const float DefaultStepRateInMilliseconds = GetDefault<UCameraCalibrationEditorSettings>()->DefaultMediaStepRateInMilliseconds;
+	const bool bForceDefaultStepRate = GetDefault<UCameraCalibrationEditorSettings>()->bForceDefaultMediaStepRate;
+
+	// The media player could return a frame rate of 0 for the current video track
+	const float MediaFrameRate = MediaPlayer->GetVideoTrackFrameRate(INDEX_NONE, INDEX_NONE);
+
+	// Use the default step rate if the media player returned an invalid frame rate or if the project settings force it
+	float MillisecondsPerStep = 0.0f;
+	if (FMath::IsNearlyEqual(MediaFrameRate, 0.0f) || bForceDefaultStepRate)
+	{
+		MillisecondsPerStep = DefaultStepRateInMilliseconds;
+	}
+	else
+	{
+		MillisecondsPerStep = 1000.0f / MediaFrameRate;
+	}
+
+
+	// Compute the number of ticks in one step and go forward from the media's current time
+	const FTimespan TicksInOneStep = ETimespan::TicksPerMillisecond * (MillisecondsPerStep);
+	const FTimespan NextStepTime = MediaPlayer->GetTime() + TicksInOneStep;
+
+	// Ensure that we do not attempt to seek past the end of the media
+	const FTimespan Duration = MediaPlayer->GetDuration();
+	if ((NextStepTime + TicksInOneStep) < Duration)
+	{
+		MediaPlayer->Seek(NextStepTime);
+		MediaPlayer->Pause();
+	}
+
+	return FReply::Handled();
+}
+
+FReply FCameraCalibrationStepsController::OnForwardButtonClicked() 
+{
+	// Increase the forward media playback rate
+	MediaPlayer->SetRate(GetFasterForwardRate());
+	return FReply::Handled(); 
+}
+
+bool FCameraCalibrationStepsController::DoesMediaSupportSeeking() const
+{
+	return MediaPlayer->SupportsSeeking();
+}
+
+bool FCameraCalibrationStepsController::DoesMediaSupportNextReverseRate() const
+{
+	constexpr bool Unthinned = false;
+	return MediaPlayer->SupportsRate(GetFasterReverseRate(), Unthinned);
+}
+
+bool FCameraCalibrationStepsController::DoesMediaSupportNextForwardRate() const
+{
+	constexpr bool Unthinned = false;
+	return MediaPlayer->SupportsRate(GetFasterForwardRate(), Unthinned);
+}
+
+float FCameraCalibrationStepsController::GetFasterReverseRate() const
+{
+	// Get the current playback rate of the media player
+	float Rate = MediaPlayer->GetRate();
+
+	// Reverse the playback direction (if needed) to ensure the rate is going in reverse
+	if (Rate > -1.0f)
+	{
+		return -1.0f;
+	}
+
+	// Double the reverse playback rate
+	return 2.0f * Rate;
+}
+
+float FCameraCalibrationStepsController::GetFasterForwardRate() const
+{
+	// Get the current playback rate of the media player
+	float Rate = MediaPlayer->GetRate();
+
+	// Reverse the playback direction (if needed) to ensure the rate is going forward
+	if (Rate < 1.0f)
+	{
+		Rate = 1.0f;
+	}
+
+	// Double the forward playback rate
+	return 2.0f * Rate;
+}
+
+void FCameraCalibrationStepsController::ToggleShowMediaPlaybackControls()
+{
+	bShowMediaPlaybackButtons = !bShowMediaPlaybackButtons;
+}
+
+bool FCameraCalibrationStepsController::AreMediaPlaybackControlsVisible() const
+{
+	return bShowMediaPlaybackButtons;
+}
+
 ULiveLinkCameraController* FCameraCalibrationStepsController::FindLiveLinkCameraController() const
 {
 	return FindLiveLinkCameraControllerWithLens(Camera.Get(), LensFile.Get());
diff --git a/Engine/Plugins/VirtualProduction/CameraCalibration/Source/CameraCalibrationEditor/Private/AssetEditor/CameraCalibrationStepsController.h b/Engine/Plugins/VirtualProduction/CameraCalibration/Source/CameraCalibrationEditor/Private/AssetEditor/CameraCalibrationStepsController.h
--- a/Engine/Plugins/VirtualProduction/CameraCalibration/Source/CameraCalibrationEditor/Private/AssetEditor/CameraCalibrationStepsController.h
+++ b/Engine/Plugins/VirtualProduction/CameraCalibration/Source/CameraCalibrationEditor/Private/AssetEditor/CameraCalibrationStepsController.h
@@ -151,6 +151,48 @@ public:
 	/** Called by the UI when the Simulcam Viewport receives keyboard input */
 	bool OnSimulcamViewportInputKey(const FKey& InKey, const EInputEvent& InEvent);
 
+	/** Called by the UI when the rewind button is clicked */
+	FReply OnRewindButtonClicked();
+
+	/** Called by the UI when the reverse button is clicked */
+	FReply OnReverseButtonClicked();
+
+	/** Called by the UI when the step back button is clicked */
+	FReply OnStepBackButtonClicked();
+
+	/** Called by the UI when the play button is clicked */
+	FReply OnPlayButtonClicked();
+
+	/** Called by the UI when the pause button is clicked */
+	FReply OnPauseButtonClicked();
+
+	/** Called by the UI when the step forward button is clicked */
+	FReply OnStepForwardButtonClicked();
+
+	/** Called by the UI when the forward button is clicked */
+	FReply OnForwardButtonClicked();
+
+	/** Returns true if the media player supports seeking */
+	bool DoesMediaSupportSeeking() const;
+
+	/** Returns true if the media player supports a reverse rate that is faster than its current rate */
+	bool DoesMediaSupportNextReverseRate() const;
+
+	/** Returns true if the media player supports a forward rate that is faster than its current rate */
+	bool DoesMediaSupportNextForwardRate() const;
+
+	/** Computes the reverse playback rate */
+	float GetFasterReverseRate() const;
+
+	/** Computes the fast forward playback rate */
+	float GetFasterForwardRate() const;
+
+	/** Toggle the setting that controls whether the media playback buttons will be visible in the UI */
+	void ToggleShowMediaPlaybackControls();
+
+	/** Returns true if the media playback buttons are visible in the UI  */
+	bool AreMediaPlaybackControlsVisible() const;
+
 private:
 
 	/** Finds the LiveLinkCameraController used in the given CameraActor that is also using the given LensFile */
@@ -257,4 +299,6 @@ private:
 	/** Pointer to the LensFileEvalData used in the current frame. Only valid during the current frame. */
 	const FLensFileEvalData* LensFileEvalData;
 
+	/** Setting to control whether the media playback buttons are visible */
+	bool bShowMediaPlaybackButtons = true;
 };
diff --git a/Engine/Plugins/VirtualProduction/CameraCalibration/Source/CameraCalibrationEditor/Private/AssetEditor/CameraCalibrationToolkit.cpp b/Engine/Plugins/VirtualProduction/CameraCalibration/Source/CameraCalibrationEditor/Private/AssetEditor/CameraCalibrationToolkit.cpp
--- a/Engine/Plugins/VirtualProduction/CameraCalibration/Source/CameraCalibrationEditor/Private/AssetEditor/CameraCalibrationToolkit.cpp
+++ b/Engine/Plugins/VirtualProduction/CameraCalibration/Source/CameraCalibrationEditor/Private/AssetEditor/CameraCalibrationToolkit.cpp
@@ -5,6 +5,7 @@
 #include "Framework/MultiBox/MultiBoxBuilder.h"
 #include "ISettingsModule.h"
 #include "CameraCalibrationStepsController.h"
+#include "CameraCalibrationCommands.h"
 #include "EngineAnalytics.h"
 #include "LensFile.h"
 #include "PropertyEditorModule.h"
@@ -55,7 +56,9 @@ void FCameraCalibrationToolkit::InitCameraCalibrationTool(const EToolkitMode::Ty
 	CalibrationStepsTab = CalibrationStepsController->BuildUI();
 	LensEditorTab = SNew(SLensFilePanel, LensFile, CalibrationStepsController.ToSharedRef())
 		.CachedFIZData(TAttribute<FCachedFIZData>::Create(TAttribute<FCachedFIZData>::FGetter::CreateSP(LensEvaluationWidget.ToSharedRef(), &SLensEvaluation::GetLastEvaluatedData)));
-		
+
+	BindCommands();
+
 	TSharedRef<FTabManager::FLayout> NewLayout = FTabManager::NewLayout("CameraCalibrationToolLayout_v0.9")
 		->AddArea
 		(
@@ -98,6 +101,7 @@ void FCameraCalibrationToolkit::InitCameraCalibrationTool(const EToolkitMode::Ty
 		bUseSmallIcons);
 
 	ExtendMenu();
+	ExtendToolBar();
 	RegenerateMenusAndToolbars();
 }
 
@@ -228,6 +232,44 @@ FCachedFIZData FCameraCalibrationToolkit::GetFIZData() const
 	return FCachedFIZData();
 }
 
+void FCameraCalibrationToolkit::BindCommands()
+{
+	const FCameraCalibrationCommands& Commands = FCameraCalibrationCommands::Get();
+
+	ToolkitCommands->MapAction(
+		Commands.ShowMediaPlaybackControls,
+		FExecuteAction::CreateRaw(CalibrationStepsController.Get(), &FCameraCalibrationStepsController::ToggleShowMediaPlaybackControls),
+		FCanExecuteAction(),
+		FIsActionChecked::CreateRaw(CalibrationStepsController.Get(), &FCameraCalibrationStepsController::AreMediaPlaybackControlsVisible));
+}
+
+void FCameraCalibrationToolkit::ExtendToolBar()
+{
+	struct Local
+	{
+		static void FillToolbar(FToolBarBuilder& ToolbarBuilder, const TSharedRef<FUICommandList> ToolkitCommands)
+		{
+			ToolbarBuilder.BeginSection("PlaybackControls");
+			{
+				ToolbarBuilder.AddToolBarButton(FCameraCalibrationCommands::Get().ShowMediaPlaybackControls);
+			}
+			ToolbarBuilder.EndSection();
+		}
+	};
+
+
+	TSharedPtr<FExtender> ToolbarExtender = MakeShareable(new FExtender);
+
+	ToolbarExtender->AddToolBarExtension(
+		"Asset",
+		EExtensionHook::After,
+		GetToolkitCommands(),
+		FToolBarExtensionDelegate::CreateStatic(&Local::FillToolbar, GetToolkitCommands())
+	);
+
+	AddToolbarExtender(ToolbarExtender);
+}
+
 void FCameraCalibrationToolkit::ExtendMenu()
 {
 	MenuExtender = MakeShared<FExtender>();
diff --git a/Engine/Plugins/VirtualProduction/CameraCalibration/Source/CameraCalibrationEditor/Private/AssetEditor/CameraCalibrationToolkit.h b/Engine/Plugins/VirtualProduction/CameraCalibration/Source/CameraCalibrationEditor/Private/AssetEditor/CameraCalibrationToolkit.h
--- a/Engine/Plugins/VirtualProduction/CameraCalibration/Source/CameraCalibrationEditor/Private/AssetEditor/CameraCalibrationToolkit.h
+++ b/Engine/Plugins/VirtualProduction/CameraCalibration/Source/CameraCalibrationEditor/Private/AssetEditor/CameraCalibrationToolkit.h
@@ -70,6 +70,12 @@ private:
 
 	FCachedFIZData GetFIZData() const;
 
+	/** Binds the UI commands to delegates. */
+	void BindCommands();
+
+	/** Builds the toolbar widget for the camera calibration tools editor. */
+	void ExtendToolBar();
+
 	/** Extend Toolkit Menu */
 	void ExtendMenu();
 
diff --git a/Engine/Plugins/VirtualProduction/CameraCalibration/Source/CameraCalibrationEditor/Private/AssetEditor/SCameraCalibrationSteps.cpp b/Engine/Plugins/VirtualProduction/CameraCalibration/Source/CameraCalibrationEditor/Private/AssetEditor/SCameraCalibrationSteps.cpp
--- a/Engine/Plugins/VirtualProduction/CameraCalibration/Source/CameraCalibrationEditor/Private/AssetEditor/SCameraCalibrationSteps.cpp
+++ b/Engine/Plugins/VirtualProduction/CameraCalibration/Source/CameraCalibrationEditor/Private/AssetEditor/SCameraCalibrationSteps.cpp
@@ -17,9 +17,11 @@
 #include "SSimulcamViewport.h"
 #include "Styling/CoreStyle.h"
 #include "Styling/SlateStyle.h"
+#include "UI/CameraCalibrationEditorStyle.h"
 #include "UI/CameraCalibrationWidgetHelpers.h"
 #include "Widgets/Colors/SColorPicker.h"
 #include "Widgets/Images/SImage.h"
+#include "Widgets/Input/SButton.h"
 #include "Widgets/Input/SCheckBox.h"
 #include "Widgets/Input/SComboBox.h"
 #include "Widgets/Input/SNumericEntryBox.h"
@@ -28,6 +30,7 @@
 #include "Widgets/Layout/SWidgetSwitcher.h"
 #include "Widgets/Layout/SBorder.h"
 #include "Widgets/Layout/SScrollBox.h"
+#include "Widgets/Layout/SUniformWrapPanel.h"
 #include "Widgets/SBoxPanel.h"
 #include "Widgets/SOverlay.h"
 #include "Widgets/Text/STextBlock.h"
@@ -67,6 +70,154 @@ void SCameraCalibrationSteps::Construct(const FArguments& InArgs, TWeakPtr<FCame
 		}
 	}
 	
+	// Make media playback buttons
+	TSharedRef<SWidget> RewindButton = SNew(SButton)
+		.ToolTipText(LOCTEXT("RewindButtonTooltip", "Rewind the media to the beginning"))
+		.OnClicked_Lambda([WeakStepsController = CalibrationStepsController]() -> FReply
+		{
+			if (TSharedPtr<FCameraCalibrationStepsController> StepsControllerShared = WeakStepsController.Pin())
+			{
+				return StepsControllerShared->OnRewindButtonClicked();
+			}
+			return FReply::Unhandled();
+		})
+		.IsEnabled_Lambda([WeakStepsController = CalibrationStepsController]() -> bool
+		{
+			if (TSharedPtr<FCameraCalibrationStepsController> StepsControllerShared = WeakStepsController.Pin())
+			{
+				return StepsControllerShared->DoesMediaSupportSeeking();
+			}
+			return false;
+		})
+		[
+			SNew(SImage)
+			.Image(FCameraCalibrationEditorStyle::Get().GetBrush(TEXT("CameraCalibration.RewindMedia.Small")))
+		];
+
+	TSharedRef<SWidget> ReverseButton = SNew(SButton)
+		.ToolTipText(LOCTEXT("ReverseButtonTooltip", "Reverse media playback"))
+		.OnClicked_Lambda([WeakStepsController = CalibrationStepsController]() -> FReply
+		{
+			if (TSharedPtr<FCameraCalibrationStepsController> StepsControllerShared = WeakStepsController.Pin())
+			{
+				return StepsControllerShared->OnReverseButtonClicked();
+			}
+			return FReply::Unhandled();
+		})
+		.IsEnabled_Lambda([WeakStepsController = CalibrationStepsController]() -> bool
+		{
+			if (TSharedPtr<FCameraCalibrationStepsController> StepsControllerShared = WeakStepsController.Pin())
+			{
+				return StepsControllerShared->DoesMediaSupportNextReverseRate();
+			}
+			return false;
+		})
+		[
+			SNew(SImage)
+			.Image(FCameraCalibrationEditorStyle::Get().GetBrush(TEXT("CameraCalibration.ReverseMedia.Small")))
+		];
+
+	TSharedRef<SWidget> StepBackButton = SNew(SButton)
+		.ToolTipText(LOCTEXT("StepBackButtonTooltip", "Step back one frame, or one time interval as set in project settings"))
+		.OnClicked_Lambda([WeakStepsController = CalibrationStepsController]() -> FReply
+		{
+			if (TSharedPtr<FCameraCalibrationStepsController> StepsControllerShared = WeakStepsController.Pin())
+			{
+				return StepsControllerShared->OnStepBackButtonClicked();
+			}
+			return FReply::Unhandled();
+		})
+		.IsEnabled_Lambda([WeakStepsController = CalibrationStepsController]() -> bool
+		{
+			if (TSharedPtr<FCameraCalibrationStepsController> StepsControllerShared = WeakStepsController.Pin())
+			{
+				return StepsControllerShared->DoesMediaSupportSeeking();
+			}
+			return false;
+		})
+		[
+			SNew(SImage)
+			.Image(FCameraCalibrationEditorStyle::Get().GetBrush(TEXT("CameraCalibration.StepBackMedia.Small")))
+		];
+
+	TSharedRef<SWidget> PlayButton = SNew(SButton)
+		.ToolTipText(LOCTEXT("PlayButtonTooltip", "Start media playback"))
+		.OnClicked_Lambda([WeakStepsController = CalibrationStepsController]() -> FReply
+		{
+			if (TSharedPtr<FCameraCalibrationStepsController> StepsControllerShared = WeakStepsController.Pin())
+			{
+				return StepsControllerShared->OnPlayButtonClicked();
+			}
+			return FReply::Unhandled();
+		})
+		[
+			SNew(SImage)
+			.Image(FCameraCalibrationEditorStyle::Get().GetBrush(TEXT("CameraCalibration.PlayMedia.Small")))
+		];
+
+	TSharedRef<SWidget> PauseButton = SNew(SButton)
+		.ToolTipText(LOCTEXT("PauseButtonTooltip", "Pause media playback"))
+		.OnClicked_Lambda([WeakStepsController = CalibrationStepsController]() -> FReply
+		{
+			if (TSharedPtr<FCameraCalibrationStepsController> StepsControllerShared = WeakStepsController.Pin())
+			{
+				return StepsControllerShared->OnPauseButtonClicked();
+			}
+			return FReply::Unhandled();
+		})
+		[
+			SNew(SImage)
+			.Image(FCameraCalibrationEditorStyle::Get().GetBrush(TEXT("CameraCalibration.PauseMedia.Small")))
+		];
+
+	TSharedRef<SWidget> StepForwardButton = SNew(SButton)
+		.ToolTipText(LOCTEXT("StepForwardButtonTooltip", "Step forward one frame, or one time interval as set in project settings"))
+		.OnClicked_Lambda([WeakStepsController = CalibrationStepsController]() -> FReply
+		{
+			if (TSharedPtr<FCameraCalibrationStepsController> StepsControllerShared = WeakStepsController.Pin())
+			{
+				return StepsControllerShared->OnStepForwardButtonClicked();
+			}
+			return FReply::Unhandled();
+		})
+		.IsEnabled_Lambda([WeakStepsController = CalibrationStepsController]() -> bool
+		{
+			if (TSharedPtr<FCameraCalibrationStepsController> StepsControllerShared = WeakStepsController.Pin())
+			{
+				return StepsControllerShared->DoesMediaSupportSeeking();
+			}
+			return false;
+		})
+		[
+			SNew(SImage)
+			.Image(FCameraCalibrationEditorStyle::Get().GetBrush(TEXT("CameraCalibration.StepForwardMedia.Small")))
+		];
+
+
+	TSharedRef<SWidget> ForwardButton = SNew(SButton)
+		.ToolTipText(LOCTEXT("ForwardButtonTooltip", "Fast forward media playback"))
+		.OnClicked_Lambda([WeakStepsController = CalibrationStepsController]() -> FReply
+		{
+			if (TSharedPtr<FCameraCalibrationStepsController> StepsControllerShared = WeakStepsController.Pin())
+			{
+				return StepsControllerShared->OnForwardButtonClicked();
+			}
+			return FReply::Unhandled();
+		})
+		.IsEnabled_Lambda([WeakStepsController = CalibrationStepsController]() -> bool
+		{
+			if (TSharedPtr<FCameraCalibrationStepsController> StepsControllerShared = WeakStepsController.Pin())
+			{
+				return StepsControllerShared->DoesMediaSupportNextForwardRate();
+			}
+			return false;
+		})
+		[
+			SNew(SImage)
+			.Image(FCameraCalibrationEditorStyle::Get().GetBrush(TEXT("CameraCalibration.ForwardMedia.Small")))
+		];
+
+
 	ChildSlot
 	[		
 		SNew(SHorizontalBox)
@@ -88,6 +239,36 @@ void SCameraCalibrationSteps::Construct(const FArguments& InArgs, TWeakPtr<FCame
 				.OnSimulcamViewportClicked_Raw(CalibrationStepsController.Pin().Get(), &FCameraCalibrationStepsController::OnSimulcamViewportClicked)
 				.OnSimulcamViewportInputKey_Raw(CalibrationStepsController.Pin().Get(), &FCameraCalibrationStepsController::OnSimulcamViewportInputKey)
 			]
+
+			+ SVerticalBox::Slot() // Media playback buttons
+			.AutoHeight()
+			.Padding(5)
+			[
+				SNew(SUniformWrapPanel)
+				.HAlign(HAlign_Center)
+				.Visibility(this, &SCameraCalibrationSteps::GetMediaPlaybackControlsVisibility)
+
+				+ SUniformWrapPanel::Slot()
+				[RewindButton]
+
+				+ SUniformWrapPanel::Slot()
+				[ReverseButton]
+
+				+ SUniformWrapPanel::Slot()
+				[StepBackButton]
+
+				+ SUniformWrapPanel::Slot()
+				[PlayButton]
+
+				+ SUniformWrapPanel::Slot()
+				[PauseButton]
+
+				+ SUniformWrapPanel::Slot()
+				[StepForwardButton]
+
+				+ SUniformWrapPanel::Slot()
+				[ForwardButton]
+			]
 		]
 
 		+ SHorizontalBox::Slot() // Right toolbar
@@ -724,5 +905,14 @@ void SCameraCalibrationSteps::SelectStep(const FName& StepName)
 	}
 }
 
+EVisibility SCameraCalibrationSteps::GetMediaPlaybackControlsVisibility() const
+{
+	TSharedPtr<FCameraCalibrationStepsController> StepsController = CalibrationStepsController.Pin();
+	if (StepsController)
+	{
+		return StepsController->AreMediaPlaybackControlsVisible() ? EVisibility::Visible : EVisibility::Collapsed;
+	}
+	return EVisibility::Collapsed;
+}
 
 #undef LOCTEXT_NAMESPACE
diff --git a/Engine/Plugins/VirtualProduction/CameraCalibration/Source/CameraCalibrationEditor/Private/AssetEditor/SCameraCalibrationSteps.h b/Engine/Plugins/VirtualProduction/CameraCalibration/Source/CameraCalibrationEditor/Private/AssetEditor/SCameraCalibrationSteps.h
--- a/Engine/Plugins/VirtualProduction/CameraCalibration/Source/CameraCalibrationEditor/Private/AssetEditor/SCameraCalibrationSteps.h
+++ b/Engine/Plugins/VirtualProduction/CameraCalibration/Source/CameraCalibrationEditor/Private/AssetEditor/SCameraCalibrationSteps.h
@@ -66,6 +66,9 @@ private:
 	/** Expected to be called when user selects a new step via the UI */
 	void SelectStep(const FName& StepName);
 
+	/** Determines the visibility of the media playback control buttons */
+	EVisibility GetMediaPlaybackControlsVisibility() const;
+
 private:
 
 	/** The controller object */
diff --git a/Engine/Plugins/VirtualProduction/CameraCalibration/Source/CameraCalibrationEditor/Private/UI/CameraCalibrationEditorStyle.h b/Engine/Plugins/VirtualProduction/CameraCalibration/Source/CameraCalibrationEditor/Private/UI/CameraCalibrationEditorStyle.h
--- a/Engine/Plugins/VirtualProduction/CameraCalibration/Source/CameraCalibrationEditor/Private/UI/CameraCalibrationEditorStyle.h
+++ b/Engine/Plugins/VirtualProduction/CameraCalibration/Source/CameraCalibrationEditor/Private/UI/CameraCalibrationEditorStyle.h
@@ -26,11 +26,22 @@ public:
 		SetContentRoot(FPaths::EngineContentDir() / TEXT("Editor/Slate/"));
 		Set("PlacementBrowser.Icons.VirtualProduction", new IMAGE_BRUSH_SVG("Starship/Common/VirtualProduction", Icon16x16));
 
+		// Toolbar
+		Set("CameraCalibration.ShowMediaPlaybackControls", new IMAGE_BRUSH_SVG("Starship/AssetIcons/MediaPlayer_16", Icon16x16));
+
 		// Set miscellaneous icons
 		SetContentRoot(FPaths::EnginePluginsDir() / TEXT("VirtualProduction/CameraCalibration/Content/Editor/Icons/"));
 		Set("ClassThumbnail.LensFile", new IMAGE_BRUSH("LensFileIcon_64x", Icon64x64));
 		Set("ClassIcon.LensFile", new IMAGE_BRUSH("LensFileIcon_20x", Icon20x20));
 
+		Set("CameraCalibration.RewindMedia.Small", new IMAGE_BRUSH("icon_rewind_40x", Icon20x20));
+		Set("CameraCalibration.ReverseMedia.Small", new IMAGE_BRUSH("icon_reverse_40x", Icon20x20));
+		Set("CameraCalibration.StepBackMedia.Small", new IMAGE_BRUSH("icon_step_back_40x", Icon20x20));
+		Set("CameraCalibration.PlayMedia.Small", new IMAGE_BRUSH("icon_play_40x", Icon20x20));
+		Set("CameraCalibration.PauseMedia.Small", new IMAGE_BRUSH("icon_pause_40x", Icon20x20));
+		Set("CameraCalibration.StepForwardMedia.Small", new IMAGE_BRUSH("icon_step_40x", Icon20x20));
+		Set("CameraCalibration.ForwardMedia.Small", new IMAGE_BRUSH("icon_forward_40x", Icon20x20));
+
 		FSlateStyleRegistry::RegisterSlateStyle(*this);
 	}
 
diff --git a/Engine/Plugins/VirtualProduction/CameraCalibrationCore/Source/CameraCalibrationCore/Public/CameraCalibrationSettings.h b/Engine/Plugins/VirtualProduction/CameraCalibrationCore/Source/CameraCalibrationCore/Public/CameraCalibrationSettings.h
--- a/Engine/Plugins/VirtualProduction/CameraCalibrationCore/Source/CameraCalibrationCore/Public/CameraCalibrationSettings.h
+++ b/Engine/Plugins/VirtualProduction/CameraCalibrationCore/Source/CameraCalibrationCore/Public/CameraCalibrationSettings.h
@@ -232,6 +232,19 @@ public:
 	UPROPERTY(config, EditAnywhere, Category = "Settings")
 	ELensDisplayUnit DefaultDisplayUnit = ELensDisplayUnit::Millimeters;
 
+	/**
+	 * If true, the media player in the calibration tools will always use the default step rate. 
+	 * Otherwise, it will try to use the frame rate of the media to step by exactly one frame. 
+	 */
+	UPROPERTY(config, EditAnywhere, Category = "Settings")
+	bool bForceDefaultMediaStepRate = false;
+
+	/**
+	 * The default step rate (ms) that the media player in the calibration tools should use when stepping forward/back
+	 */
+	UPROPERTY(config, EditAnywhere, Category = "Settings")
+	float DefaultMediaStepRateInMilliseconds = 100.0f;
+
 private:
 
 	/** 
