[CameraCalibration] Auto-export calibration datasets to disk and provide mechanism to import them from disk to support debugging and reproducing bad calibration results.

diff --git a/Engine/Plugins/VirtualProduction/CameraCalibration/Source/CameraCalibrationEditor/CameraCalibrationEditor.Build.cs b/Engine/Plugins/VirtualProduction/CameraCalibration/Source/CameraCalibrationEditor/CameraCalibrationEditor.Build.cs
--- a/Engine/Plugins/VirtualProduction/CameraCalibration/Source/CameraCalibrationEditor/CameraCalibrationEditor.Build.cs
+++ b/Engine/Plugins/VirtualProduction/CameraCalibration/Source/CameraCalibrationEditor/CameraCalibrationEditor.Build.cs
@@ -20,9 +20,11 @@ namespace UnrealBuildTool.Rules
 					"Core",
 					"CoreUObject",
 	                "CurveEditor",
+					"DesktopPlatform",
 					"EditorStyle",
 					"EditorWidgets",
 					"Engine",
+					"ImageCore",
 					"InputCore",
 					"Json",
 					"JsonUtilities",
diff --git a/Engine/Plugins/VirtualProduction/CameraCalibration/Source/CameraCalibrationEditor/Private/AssetEditor/LensDistortionTool.cpp b/Engine/Plugins/VirtualProduction/CameraCalibration/Source/CameraCalibrationEditor/Private/AssetEditor/LensDistortionTool.cpp
--- a/Engine/Plugins/VirtualProduction/CameraCalibration/Source/CameraCalibrationEditor/Private/AssetEditor/LensDistortionTool.cpp
+++ b/Engine/Plugins/VirtualProduction/CameraCalibration/Source/CameraCalibrationEditor/Private/AssetEditor/LensDistortionTool.cpp
@@ -2,20 +2,37 @@
 
 #include "LensDistortionTool.h"
 
+#include "CameraCalibrationEditorLog.h"
+#include "CameraCalibrationSettings.h"
 #include "CameraCalibrationStepsController.h"
 #include "CameraCalibrationTypes.h"
 #include "CameraLensDistortionAlgo.h"
+#include "DesktopPlatformModule.h"
+#include "Dom/JsonObject.h"
+#include "HAL/FileManager.h"
+#include "IDesktopPlatform.h"
+#include "ImageUtils.h"
 #include "LensFile.h"
 #include "LensInfoStep.h"
+#include "Misc/DateTime.h"
 #include "Misc/MessageDialog.h"
 #include "Models/SphericalLensModel.h"
 #include "ScopedTransaction.h"
+#include "Serialization/JsonWriter.h"
+#include "Serialization/JsonSerializer.h"
 #include "SLensDistortionToolPanel.h"
 #include "UObject/Class.h"
 #include "UObject/UObjectIterator.h"
 
 #define LOCTEXT_NAMESPACE "LensDistortionTool"
 
+namespace UE::CameraCalibration::Private::LensDistortionTool
+{
+	static const FString SaveDir = FPaths::ProjectSavedDir() / TEXT("CameraCalibration") / TEXT("LensDistortion");
+	static const FString SessionDateTimeField(TEXT("SessionDateTime"));
+	static const FString AlgoNameField(TEXT("AlgoName"));
+}
+
 void ULensDistortionTool::Initialize(TWeakPtr<FCameraCalibrationStepsController> InCameraCalibrationStepController)
 {
 	CameraCalibrationStepsController = InCameraCalibrationStepController;
@@ -46,6 +63,8 @@ void ULensDistortionTool::Shutdown()
 	{
 		CurrentAlgo->Shutdown();
 		CurrentAlgo = nullptr;
+
+		EndCalibrationSession();
 	}
 }
 
@@ -126,6 +145,8 @@ void ULensDistortionTool::SetAlgo(const FName& AlgoName)
 	{
 		CurrentAlgo->Shutdown();
 		CurrentAlgo = nullptr;
+
+		EndCalibrationSession();
 	}
 
 	// If AlgoClass is none, we're done here.
@@ -279,4 +300,308 @@ bool ULensDistortionTool::IsOverlayEnabled() const
 	return false;
 }
 
+void ULensDistortionTool::StartCalibrationSession()
+{
+	if (!SessionInfo.bIsActive)
+	{
+		SessionInfo.bIsActive = true;
+		SessionInfo.StartTime = FDateTime::Now();
+	}
+}
+
+void ULensDistortionTool::EndCalibrationSession()
+{
+	if (SessionInfo.bIsActive)
+	{
+		SessionInfo.bIsActive = false;
+		SessionInfo.RowIndex = -1;
+	}
+}
+
+uint32 ULensDistortionTool::AdvanceSessionRowIndex()
+{
+	SessionInfo.RowIndex += 1;
+	return SessionInfo.RowIndex;
+}
+
+FString ULensDistortionTool::GetSessionSaveDir() const
+{
+	using namespace UE::CameraCalibration::Private;
+
+	const FString SessionDateString = SessionInfo.StartTime.ToString(TEXT("%Y-%m-%d"));
+	const FString SessionTimeString = SessionInfo.StartTime.ToString(TEXT("%H-%M-%S"));
+	const FString DatasetPrefix = TEXT("Dataset-") + CurrentAlgo->ShortName().ToString() + TEXT("Algorithm-");
+	const FString DatasetDir = DatasetPrefix + SessionTimeString;
+
+	return LensDistortionTool::SaveDir / SessionDateString / DatasetDir;
+}
+
+FString ULensDistortionTool::GetRowFilename(int32 RowIndex) const
+{
+	const FString RowNumString = TEXT("Row") + FString::FromInt(RowIndex) + TEXT("-");
+	return RowNumString;
+}
+
+void ULensDistortionTool::DeleteExportedRow(const int32& RowIndex) const
+{
+	if (!GetDefault<UCameraCalibrationSettings>()->IsCalibrationDatasetImportExportEnabled())
+	{
+		return;
+	}
+
+	// Find all files in the directory of the currently active session
+	const FString PathName = GetSessionSaveDir();
+	TArray<FString> FoundFiles;
+	IFileManager::Get().FindFiles(FoundFiles, *PathName);
+
+	const FString RowNumString = GetRowFilename(RowIndex);
+
+	// Delete any files containing that row number from the session directory
+	for (const FString& File : FoundFiles)
+	{
+		if (File.Contains(RowNumString))
+		{
+			const FString FullPath = PathName / File;
+			IFileManager::Get().Delete(*FullPath);
+			UE_LOG(LogCameraCalibrationEditor, Verbose, TEXT("Lens Distortion Tool removed calibration dataset file: %s"), *FullPath);
+		}
+	}
+}
+
+void ULensDistortionTool::ExportCalibrationRow(int32 RowIndex, const TSharedRef<FJsonObject>& RowObject, const FImageView& RowImage)
+{
+	if (!GetDefault<UCameraCalibrationSettings>()->IsCalibrationDatasetImportExportEnabled())
+	{
+		return;
+	}
+
+	// Start a calibration session (if one is not currently active)
+	StartCalibrationSession();
+
+	// Assemble the path and filename for this row based on the session and row index
+	const FString PathName = GetSessionSaveDir();
+	const FString FileName = GetRowFilename(RowIndex) + FDateTime::Now().ToString(TEXT("%H-%M-%S"));
+
+	const FString JsonFileName = PathName / FileName + TEXT(".json");
+	const FString ImageFileName = PathName / FileName + TEXT(".png");
+
+	// Create and open a new Json file for writing, and initialize a JsonWriter to serialize the contents
+	if (TUniquePtr<FArchive> FileWriter = TUniquePtr<FArchive>(IFileManager::Get().CreateFileWriter(*JsonFileName)))
+	{
+		TSharedRef< TJsonWriter<> > JsonWriter = TJsonWriterFactory<>::Create(FileWriter.Get());
+
+		// Write the Json row data out and save the file
+		FJsonSerializer::Serialize(RowObject, JsonWriter);
+		FileWriter->Close();
+
+		UE_LOG(LogCameraCalibrationEditor, Verbose, TEXT("Lens Distortion Tool wrote to dataset row file: %s"), *JsonFileName);
+	}
+
+	// If the row has an image to export, save it out to a file
+	if (RowImage.RawData != nullptr)
+	{
+		FImageUtils::SaveImageByExtension(*ImageFileName, RowImage);
+	}
+}
+
+void ULensDistortionTool::ExportSessionData(const TSharedRef<FJsonObject>& SessionDataObject)
+{
+	using namespace UE::CameraCalibration::Private;
+
+	if (!GetDefault<UCameraCalibrationSettings>()->IsCalibrationDatasetImportExportEnabled())
+	{
+		return;
+	}
+
+	// Start a calibration session (if one is not currently active)
+	StartCalibrationSession();
+
+	// Assemble the path and filename for this row based on the session and row index
+	const FString PathName = GetSessionSaveDir();
+	const FString FileName = TEXT("SessionData");
+
+	const FString SessionFileName = PathName / FileName + TEXT(".ucamcalib");
+
+	// Delete the existing session data file (if it exists)
+	if (IFileManager::Get().FileExists(*SessionFileName))
+	{
+		IFileManager::Get().Delete(*SessionFileName);
+	}
+
+	// Create and open a new Json file for writing, and initialize a JsonWriter to serialize the contents
+	if (TUniquePtr<FArchive> FileWriter = TUniquePtr<FArchive>(IFileManager::Get().CreateFileWriter(*SessionFileName)))
+	{
+		TSharedRef< TJsonWriter<> > JsonWriter = TJsonWriterFactory<>::Create(FileWriter.Get());
+
+		const FString SessionDateTimeString = SessionInfo.StartTime.ToString(TEXT("%Y-%m-%d")) + TEXT("-") + SessionInfo.StartTime.ToString(TEXT("%H-%M-%S"));
+
+		SessionDataObject->SetStringField(LensDistortionTool::SessionDateTimeField, SessionDateTimeString);
+		SessionDataObject->SetStringField(LensDistortionTool::AlgoNameField, CurrentAlgo->FriendlyName().ToString());
+
+		// Write the Json row data out and save the file
+		FJsonSerializer::Serialize(SessionDataObject, JsonWriter);
+		FileWriter->Close();
+
+		UE_LOG(LogCameraCalibrationEditor, Verbose, TEXT("Lens Distortion Tool wrote to dataset session file: %s"), *SessionFileName);
+	}
+}
+
+void ULensDistortionTool::ImportCalibrationDataset()
+{
+	using namespace UE::CameraCalibration::Private;
+
+	if (!GetDefault<UCameraCalibrationSettings>()->IsCalibrationDatasetImportExportEnabled())
+	{
+		return;
+	}
+
+	// If there is existing calibration data that will be overwritten during import, ask the user to confirm that they want to continue
+	if (CurrentAlgo->HasCalibrationData())
+	{
+		const FText ConfirmationMessage = LOCTEXT(
+			"ImportDatasetConfirmationMessage",
+			"There are existing calibration rows which will be removed during the import process. Do you want to proceed with the import?");
+
+		if (FMessageDialog::Open(EAppMsgType::YesNo, ConfirmationMessage) == EAppReturnType::No)
+		{
+			return;
+		}
+	}
+
+	// Open a file dialog to select a .ucamcalib session data file 
+	IDesktopPlatform* DesktopPlatform = FDesktopPlatformModule::Get();
+	const void* ParentWindowHandle = FSlateApplication::Get().FindBestParentWindowHandleForDialogs(nullptr);
+	const FString Title = TEXT("Import Camera Calibration Dataset");
+	const FString DefaultPath = LensDistortionTool::SaveDir;
+	const FString DefaultFile = TEXT("");
+	const FString FileTypes = TEXT("Camera Calibration Dataset|*.ucamcalib");
+	const uint32 OpenFileFlags = 0;
+
+	// Note, OpenFileFlags is not set to "Multiple" so we only expect one file to be selected
+	TArray<FString> SelectedFileNames;
+	const bool bFileSelected = DesktopPlatform->OpenFileDialog(ParentWindowHandle, Title, DefaultPath, DefaultFile, FileTypes, OpenFileFlags, SelectedFileNames);
+
+	// Early-out if no calibration file was selected
+	if (!bFileSelected || SelectedFileNames.Num() < 1)
+	{
+		return;
+	}
+
+	// Parse the session data filename and the directory from the full path
+	const FString SessionFileName = FPaths::GetCleanFilename(SelectedFileNames[0]);
+	const FString SelectedDirectory = FPaths::GetPath(SelectedFileNames[0]);
+
+	// Find all json files in the selected directory (this will not include the .ucamcalib session date file)
+	TArray<FString> FoundFiles;
+	const FString FileExtension = TEXT(".json");
+	IFileManager::Get().FindFiles(FoundFiles, *SelectedDirectory, *FileExtension);
+
+	// Early-out if selected directory has no json files to import
+	if (FoundFiles.Num() < 1)
+	{
+		const FText ErrorMessage = LOCTEXT("NoJsonFilesFound", "The selected directory has no .json files to import.");
+		FMessageDialog::Open(EAppMsgType::Ok, ErrorMessage);
+		return;
+	}
+
+	// Import the session data
+	FDateTime ImportedSessionDateTime = FDateTime::Now();
+	{
+		const FString SessionFile = SelectedDirectory / SessionFileName;
+
+		// Open the Json file for reading, and initialize a JsonReader to parse the contents
+		if (TUniquePtr<FArchive> FileReader = TUniquePtr<FArchive>(IFileManager::Get().CreateFileReader(*SessionFile)))
+		{
+			TSharedRef< TJsonReader<> > JsonReader = TJsonReaderFactory<>::Create(FileReader.Get());
+
+			// Deserialize the row data from the Json file into a Json object
+			TSharedPtr<FJsonObject> JsonSessionData = MakeShared<FJsonObject>();
+			if (FJsonSerializer::Deserialize(JsonReader, JsonSessionData))
+			{
+				// Import the session date/time so that we can restore the imported session
+				FString SessionDateTimeString;
+				if (JsonSessionData->TryGetStringField(LensDistortionTool::SessionDateTimeField, SessionDateTimeString))
+				{
+					ensureMsgf(FDateTime::Parse(SessionDateTimeString, ImportedSessionDateTime), TEXT("Failed to parse imported session date and time"));
+				}
+				else
+				{
+					UE_LOG(LogCameraCalibrationEditor, Verbose, TEXT("Lens Distortion Tool failed to deserialize the date/time from the session file: %s"), *SessionFile);
+				}
+
+				// Import the algo name so we can set the appropriate algo before importing the row data
+				FString AlgoString;
+				if (JsonSessionData->TryGetStringField(LensDistortionTool::AlgoNameField, AlgoString))
+				{
+					// Ensure that the algo name matches one of the algos for this tool
+					const FName AlgoName = FName(*AlgoString);
+					if (AlgosMap.Contains(AlgoName))
+					{
+						SetAlgo(AlgoName);
+					}
+					else
+					{
+						const FText ErrorMessage = LOCTEXT("UnknownCalibrationAlgo", "The selected dataset does not represent a lens distortion calibration. Choose a different .ucamcalib dataset file from a lens distortion calibration.");
+						FMessageDialog::Open(EAppMsgType::Ok, ErrorMessage);
+						return;
+					}
+				}
+				else
+				{
+					UE_LOG(LogCameraCalibrationEditor, Verbose, TEXT("Lens Distortion Tool failed to deserialize the algo name from the session file: %s"), *SessionFile);
+				}
+
+				CurrentAlgo->ImportSessionData(JsonSessionData.ToSharedRef());
+			}
+			else
+			{
+				UE_LOG(LogCameraCalibrationEditor, Verbose, TEXT("Lens Distortion Tool failed to deserialize dataset session file: %s"), *SessionFile);
+			}
+		}
+	}
+
+	CurrentAlgo->PreImportCalibrationData();
+
+	// Initialize a maximum row index which will be used to set the current session row index if the user wants to add additional rows after importing
+	int32 MaxRowIndex = -1;
+
+	for (const FString& File : FoundFiles)
+	{
+		const FString JsonFileName = SelectedDirectory / File;
+		const FString ImageFileName = JsonFileName.Replace(TEXT(".json"), TEXT(".png"));
+
+		// Load the PNG image file for this row into an FImage
+		FImage RowImage;
+		if (IFileManager::Get().FileExists(*ImageFileName))
+		{
+			FImageUtils::LoadImage(*ImageFileName, RowImage);
+		}
+
+		// Open the Json file for reading, and initialize a JsonReader to parse the contents
+		if (TUniquePtr<FArchive> FileReader = TUniquePtr<FArchive>(IFileManager::Get().CreateFileReader(*JsonFileName)))
+		{
+			TSharedRef< TJsonReader<> > JsonReader = TJsonReaderFactory<>::Create(FileReader.Get());
+
+			// Deserialize the row data from the Json file into a Json object
+			TSharedPtr<FJsonObject> JsonRowData = MakeShared<FJsonObject>();
+			if (FJsonSerializer::Deserialize(JsonReader, JsonRowData))
+			{
+				int32 RowNum = CurrentAlgo->ImportCalibrationRow(JsonRowData.ToSharedRef(), RowImage);
+				MaxRowIndex = FMath::Max(MaxRowIndex, RowNum);
+			}
+			else
+			{
+				UE_LOG(LogCameraCalibrationEditor, Verbose, TEXT("Lens Distortion Tool failed to deserialize the dataset row file: %s"), *JsonFileName);
+			}
+		}
+	}
+
+	CurrentAlgo->PostImportCalibrationData();
+
+	// Set the current session's start date/time and row index to match what was just imported to support adding/deleting rows
+	SessionInfo.bIsActive = true;
+	SessionInfo.StartTime = ImportedSessionDateTime;
+	SessionInfo.RowIndex = MaxRowIndex;
+}
+
 #undef LOCTEXT_NAMESPACE
diff --git a/Engine/Plugins/VirtualProduction/CameraCalibration/Source/CameraCalibrationEditor/Private/AssetEditor/LensDistortionTool.h b/Engine/Plugins/VirtualProduction/CameraCalibration/Source/CameraCalibrationEditor/Private/AssetEditor/LensDistortionTool.h
--- a/Engine/Plugins/VirtualProduction/CameraCalibration/Source/CameraCalibrationEditor/Private/AssetEditor/LensDistortionTool.h
+++ b/Engine/Plugins/VirtualProduction/CameraCalibration/Source/CameraCalibrationEditor/Private/AssetEditor/LensDistortionTool.h
@@ -6,6 +6,7 @@
 
 #include "Engine/World.h"
 #include "CameraCalibrationStep.h"
+#include "ImageCore.h"
 
 #include "LensDistortionTool.generated.h"
 
@@ -13,8 +14,22 @@ struct FGeometry;
 struct FLensFileEvalData;
 struct FPointerEvent;
 
+class FJsonObject;
 class UCameraLensDistortionAlgo;
 
+/** Data associated with a lens distortion calibration session */
+struct FLensDistortionSessionInfo
+{
+	/** The date/time when the current calibration session started */
+	FDateTime StartTime;
+
+	/** The index of the next row in the current calibration session */
+	int32 RowIndex = -1;
+
+	/** True if a calibration session is currently in progress */
+	bool bIsActive = false;
+};
+
 /**
  * ULensDistortionTool is the controller for the lens distortion panel.
  */
@@ -41,8 +56,6 @@ public:
 	virtual bool IsOverlayEnabled() const override;
 	//~ End UCameraCalibrationStep interface
 
-public:
-
 	/** Selects the algorithm by name */
 	void SetAlgo(const FName& AlgoName);
 
@@ -52,11 +65,41 @@ public:
 	/** Returns available algorithm names */
 	TArray<FName> GetAlgos() const;
 
-public:
-
 	/** Called by the UI when the user wants to save the calibration data that the current algorithm is providing */
 	void OnSaveCurrentCalibrationData();
 
+	/** Initiate a new calibration session (if one is not already active) */
+	void StartCalibrationSession();
+
+	/** End the current calibration session */
+	void EndCalibrationSession();
+
+	/** Increments the session index and returns its new value */
+	uint32 AdvanceSessionRowIndex();
+
+	/** Intended to be called by one of the algo objects. Exports the data needed by the algo for the current session to a .json file on disk. */
+	void ExportSessionData(const TSharedRef<FJsonObject>& SessionDataObject);
+
+	/** Intended to be called by one of the algo objects. Exports the data associated with a single calibration row to a .json file on disk. */
+	void ExportCalibrationRow(int32 RowIndex, const TSharedRef<FJsonObject>& RowObject, const FImageView& RowImage = FImageView());
+
+	/** Intended to be called by one of the algo objects. Deletes the .json file with the input row index that was previously exported for this session. */
+	void DeleteExportedRow(const int32& RowIndex) const;
+
+	/** Import a saved calibration data set from disk */
+	void ImportCalibrationDataset();
+
+private:
+	/** Get the directory for the current session */
+	FString GetSessionSaveDir() const;
+
+	/** Get the filename of a row with the input index */
+	FString GetRowFilename(int32 RowIndex) const;
+
+public:
+	/** Stores info about the current calibration session of this tool */
+	FLensDistortionSessionInfo SessionInfo;
+
 private:
 
 	/** Pointer to the calibration steps controller */
diff --git a/Engine/Plugins/VirtualProduction/CameraCalibration/Source/CameraCalibrationEditor/Private/AssetEditor/NodalOffsetTool.cpp b/Engine/Plugins/VirtualProduction/CameraCalibration/Source/CameraCalibrationEditor/Private/AssetEditor/NodalOffsetTool.cpp
--- a/Engine/Plugins/VirtualProduction/CameraCalibration/Source/CameraCalibrationEditor/Private/AssetEditor/NodalOffsetTool.cpp
+++ b/Engine/Plugins/VirtualProduction/CameraCalibration/Source/CameraCalibrationEditor/Private/AssetEditor/NodalOffsetTool.cpp
@@ -4,17 +4,34 @@
 
 #include "CalibrationPointComponent.h"
 #include "Camera/CameraActor.h"
+#include "CameraCalibrationEditorLog.h"
+#include "CameraCalibrationSettings.h"
 #include "CameraCalibrationStepsController.h"
 #include "CameraCalibrationSubsystem.h"
 #include "CameraCalibrationTypes.h"
+#include "DesktopPlatformModule.h"
+#include "Dom/JsonObject.h"
+#include "HAL/FileManager.h"
+#include "IDesktopPlatform.h"
 #include "ImageCenterTool.h"
+#include "ImageUtils.h"
 #include "LensComponent.h"
+#include "Misc/DateTime.h"
 #include "Misc/MessageDialog.h"
 #include "ScopedTransaction.h"
+#include "Serialization/JsonWriter.h"
+#include "Serialization/JsonSerializer.h"
 #include "SNodalOffsetToolPanel.h"
 
 #define LOCTEXT_NAMESPACE "NodalOffsetTool"
 
+namespace UE::CameraCalibration::Private::NodalOffsetTool
+{
+	static const FString SaveDir = FPaths::ProjectSavedDir() / TEXT("CameraCalibration") / TEXT("NodalOffset");
+	static const FString SessionDateTimeField(TEXT("SessionDateTime"));
+	static const FString AlgoNameField(TEXT("AlgoName"));
+}
+
 void UNodalOffsetTool::Initialize(TWeakPtr<FCameraCalibrationStepsController> InCameraCalibrationStepController)
 {
 	CameraCalibrationStepsController = InCameraCalibrationStepController;
@@ -41,6 +58,8 @@ void UNodalOffsetTool::Shutdown()
 	{
 		NodalOffsetAlgo->Shutdown();
 		NodalOffsetAlgo = nullptr;
+
+		EndCalibrationSession();
 	}
 }
 
@@ -132,6 +151,8 @@ void UNodalOffsetTool::SetNodalOffsetAlgo(const FName& AlgoName)
 	{
 		NodalOffsetAlgo->Shutdown();
 		NodalOffsetAlgo = nullptr;
+
+		EndCalibrationSession();
 	}
 
 	// If AlgoClass is none, we're done here.
@@ -270,4 +291,311 @@ bool UNodalOffsetTool::IsOverlayEnabled() const
 	return NodalOffsetAlgo->IsOverlayEnabled();
 }
 
+void UNodalOffsetTool::StartCalibrationSession()
+{
+	if (!SessionInfo.bIsActive)
+	{
+		SessionInfo.bIsActive = true;
+		SessionInfo.StartTime = FDateTime::Now();
+	}
+}
+
+void UNodalOffsetTool::EndCalibrationSession()
+{
+	if (SessionInfo.bIsActive)
+	{
+		SessionInfo.bIsActive = false;
+		SessionInfo.RowIndex = -1;
+	}
+}
+
+uint32 UNodalOffsetTool::AdvanceSessionRowIndex()
+{
+	SessionInfo.RowIndex += 1;
+	return SessionInfo.RowIndex;
+}
+
+FString UNodalOffsetTool::GetSessionSaveDir() const
+{
+	using namespace UE::CameraCalibration::Private;
+
+	const FString SessionDateString = SessionInfo.StartTime.ToString(TEXT("%Y-%m-%d"));
+	const FString SessionTimeString = SessionInfo.StartTime.ToString(TEXT("%H-%M-%S"));
+	const FString DatasetPrefix = TEXT("Dataset-") + NodalOffsetAlgo->ShortName().ToString() + TEXT("Algorithm-");
+	const FString DatasetDir = DatasetPrefix + SessionTimeString;
+
+	return NodalOffsetTool::SaveDir / SessionDateString / DatasetDir;
+}
+
+FString UNodalOffsetTool::GetRowFilename(int32 RowIndex) const
+{
+	const FString RowNumString = TEXT("Row") + FString::FromInt(RowIndex) + TEXT("-");
+	return RowNumString;
+}
+
+void UNodalOffsetTool::DeleteExportedRow(const int32& RowIndex) const
+{
+	if (!GetDefault<UCameraCalibrationSettings>()->IsCalibrationDatasetImportExportEnabled())
+	{
+		return;
+	}
+
+	// Find all files in the directory of the currently active session
+	const FString PathName = GetSessionSaveDir();
+	TArray<FString> FoundFiles;
+	IFileManager::Get().FindFiles(FoundFiles, *PathName);
+
+	const FString RowNumString = GetRowFilename(RowIndex);
+
+	// Delete any files containing that row number from the session directory
+	for (const FString& File : FoundFiles)
+	{
+		if (File.Contains(RowNumString))
+		{
+			const FString FullPath = PathName / File;
+			IFileManager::Get().Delete(*FullPath);
+			UE_LOG(LogCameraCalibrationEditor, Verbose, TEXT("Nodal Offset Tool removed calibration dataset file: %s"), *FullPath);
+		}
+	}
+}
+
+void UNodalOffsetTool::ExportCalibrationRow(int32 RowIndex, const TSharedRef<FJsonObject>& RowObject, const FImageView& RowImage)
+{
+	if (!GetDefault<UCameraCalibrationSettings>()->IsCalibrationDatasetImportExportEnabled())
+	{
+		return;
+	}
+
+	// Start a calibration session (if one is not currently active)
+	StartCalibrationSession();
+
+	// Assemble the path and filename for this row based on the session and row index
+	const FString PathName = GetSessionSaveDir();
+	const FString FileName = GetRowFilename(RowIndex) + FDateTime::Now().ToString(TEXT("%H-%M-%S"));
+
+	const FString JsonFileName = PathName / FileName + TEXT(".json");
+	const FString ImageFileName = PathName / FileName + TEXT(".png");
+
+	// Create and open a new Json file for writing, and initialize a JsonWriter to serialize the contents
+	if (TUniquePtr<FArchive> FileWriter = TUniquePtr<FArchive>(IFileManager::Get().CreateFileWriter(*JsonFileName)))
+	{
+		TSharedRef< TJsonWriter<> > JsonWriter = TJsonWriterFactory<>::Create(FileWriter.Get());
+
+		// Write the Json row data out and save the file
+		FJsonSerializer::Serialize(RowObject, JsonWriter);
+		FileWriter->Close();
+
+		UE_LOG(LogCameraCalibrationEditor, Verbose, TEXT("Nodal Offset Tool wrote to dataset row file: %s"), *JsonFileName);
+	}
+
+	// If the row has an image to export, save it out to a file
+	if (RowImage.RawData != nullptr)
+	{
+		FImageUtils::SaveImageByExtension(*ImageFileName, RowImage);
+	}
+}
+
+void UNodalOffsetTool::ExportSessionData(const TSharedRef<FJsonObject>& SessionDataObject)
+{
+	using namespace UE::CameraCalibration::Private;
+
+	if (!GetDefault<UCameraCalibrationSettings>()->IsCalibrationDatasetImportExportEnabled())
+	{
+		return;
+	}
+
+	// Start a calibration session (if one is not currently active)
+	StartCalibrationSession();
+
+	// Assemble the path and filename for this row based on the session and row index
+	const FString PathName = GetSessionSaveDir();
+	const FString FileName = TEXT("SessionData");
+
+	const FString SessionFileName = PathName / FileName + TEXT(".ucamcalib");
+
+	// Delete the existing session data file (if it exists)
+	if (IFileManager::Get().FileExists(*SessionFileName))
+	{
+		IFileManager::Get().Delete(*SessionFileName);
+	}
+
+	// Create and open a new Json file for writing, and initialize a JsonWriter to serialize the contents
+	if (TUniquePtr<FArchive> FileWriter = TUniquePtr<FArchive>(IFileManager::Get().CreateFileWriter(*SessionFileName)))
+	{
+		TSharedRef< TJsonWriter<> > JsonWriter = TJsonWriterFactory<>::Create(FileWriter.Get());
+
+		const FString SessionDateTimeString = SessionInfo.StartTime.ToString(TEXT("%Y-%m-%d")) + TEXT("-") + SessionInfo.StartTime.ToString(TEXT("%H-%M-%S"));
+
+		SessionDataObject->SetStringField(NodalOffsetTool::SessionDateTimeField, SessionDateTimeString);
+		SessionDataObject->SetStringField(NodalOffsetTool::AlgoNameField, NodalOffsetAlgo->FriendlyName().ToString());
+
+		// Write the Json row data out and save the file
+		FJsonSerializer::Serialize(SessionDataObject, JsonWriter);
+		FileWriter->Close();
+
+		UE_LOG(LogCameraCalibrationEditor, Verbose, TEXT("Nodal Offset Tool wrote to dataset session file: %s"), *SessionFileName);
+	}
+}
+
+void UNodalOffsetTool::ImportCalibrationDataset()
+{
+	using namespace UE::CameraCalibration::Private;
+
+	if (!GetDefault<UCameraCalibrationSettings>()->IsCalibrationDatasetImportExportEnabled())
+	{
+		return;
+	}
+
+	// If there is existing calibration data that will be overwritten during import, ask the user to confirm that they want to continue
+	if (NodalOffsetAlgo->HasCalibrationData())
+	{
+		const FText ConfirmationMessage = LOCTEXT(
+			"ImportDatasetConfirmationMessage",
+			"There are existing calibration rows which will be removed during the import process. Do you want to proceed with the import?");
+
+		if (FMessageDialog::Open(EAppMsgType::YesNo, ConfirmationMessage) == EAppReturnType::No)
+		{
+			return;
+		}
+	}
+
+	// Open a file dialog to select a .ucamcalib session data file 
+	IDesktopPlatform* DesktopPlatform = FDesktopPlatformModule::Get();
+	const void* ParentWindowHandle = FSlateApplication::Get().FindBestParentWindowHandleForDialogs(nullptr);
+	const FString Title = TEXT("Import Camera Calibration Dataset");
+	const FString DefaultPath = NodalOffsetTool::SaveDir;
+	const FString DefaultFile = TEXT("");
+	const FString FileTypes = TEXT("Camera Calibration Dataset|*.ucamcalib");
+	const uint32 OpenFileFlags = 0;
+
+	// Note, OpenFileFlags is not set to "Multiple" so we only expect one file to be selected
+	TArray<FString> SelectedFileNames;
+	const bool bFileSelected = DesktopPlatform->OpenFileDialog(ParentWindowHandle, Title, DefaultPath, DefaultFile, FileTypes, OpenFileFlags, SelectedFileNames);
+
+	// Early-out if no calibration file was selected
+	if (!bFileSelected || SelectedFileNames.Num() < 1)
+	{
+		return;
+	}
+
+	// Parse the session data filename and the directory from the full path
+	const FString SessionFileName = FPaths::GetCleanFilename(SelectedFileNames[0]);
+	const FString SelectedDirectory = FPaths::GetPath(SelectedFileNames[0]);
+
+	// Find all json files in the selected directory (this will not include the .ucamcalib session date file)
+	TArray<FString> FoundFiles;
+	const FString FileExtension = TEXT(".json");
+	IFileManager::Get().FindFiles(FoundFiles, *SelectedDirectory, *FileExtension);
+
+	// Early-out if selected directory has no json files to import
+	if (FoundFiles.Num() < 1)
+	{
+		const FText ErrorMessage = LOCTEXT("NoJsonFilesFound", "The selected directory has no .json files to import.");
+		FMessageDialog::Open(EAppMsgType::Ok, ErrorMessage);
+		return;
+	}
+
+	// Import the session data
+	FDateTime ImportedSessionDateTime = FDateTime::Now();
+	{
+		const FString SessionFile = SelectedDirectory / SessionFileName;
+
+		// Open the Json file for reading, and initialize a JsonReader to parse the contents
+		if (TUniquePtr<FArchive> FileReader = TUniquePtr<FArchive>(IFileManager::Get().CreateFileReader(*SessionFile)))
+		{
+			TSharedRef< TJsonReader<> > JsonReader = TJsonReaderFactory<>::Create(FileReader.Get());
+
+			// Deserialize the row data from the Json file into a Json object
+			TSharedPtr<FJsonObject> JsonSessionData = MakeShared<FJsonObject>();
+			if (FJsonSerializer::Deserialize(JsonReader, JsonSessionData))
+			{
+				// Import the session date/time so that we can restore the imported session
+				FString SessionDateTimeString;
+				if (JsonSessionData->TryGetStringField(NodalOffsetTool::SessionDateTimeField, SessionDateTimeString))
+				{
+					ensureMsgf(FDateTime::Parse(SessionDateTimeString, ImportedSessionDateTime), TEXT("Failed to parse imported session date and time"));
+				}
+				else
+				{
+					UE_LOG(LogCameraCalibrationEditor, Verbose, TEXT("Nodal Offset Tool failed to deserialize the date/time from the session file: %s"), *SessionFile);
+				}
+
+				// Import the algo name so we can set the appropriate algo before importing the row data
+				FString AlgoString;
+				if (JsonSessionData->TryGetStringField(NodalOffsetTool::AlgoNameField, AlgoString))
+				{
+					// Ensure that the algo name matches one of the algos for this tool
+					const FName AlgoName = FName(*AlgoString);
+					if (UCameraCalibrationSubsystem* Subsystem = GEngine->GetEngineSubsystem<UCameraCalibrationSubsystem>())
+					{
+						if (TSubclassOf<UCameraNodalOffsetAlgo> AlgoClass = Subsystem->GetCameraNodalOffsetAlgo(AlgoName))
+						{
+							SetNodalOffsetAlgo(AlgoName);
+						}
+						else
+						{
+							const FText ErrorMessage = LOCTEXT("UnknownCalibrationAlgo", "The selected dataset does not represent a nodal offset calibration. Choose a different .ucamcalib dataset file from a nodal offset calibration.");
+							FMessageDialog::Open(EAppMsgType::Ok, ErrorMessage);
+							return;
+						}
+					}
+				}
+				else
+				{
+					UE_LOG(LogCameraCalibrationEditor, Verbose, TEXT("Nodal Offset Tool failed to deserialize the algo name from the session file: %s"), *SessionFile);
+				}
+
+				NodalOffsetAlgo->ImportSessionData(JsonSessionData.ToSharedRef());
+			}
+			else
+			{
+				UE_LOG(LogCameraCalibrationEditor, Verbose, TEXT("Nodal Offset Tool failed to deserialize dataset session file: %s"), *SessionFile);
+			}
+		}
+	}
+
+	NodalOffsetAlgo->PreImportCalibrationData();
+
+	// Initialize a maximum row index which will be used to set the current session row index if the user wants to add additional rows after importing
+	int32 MaxRowIndex = -1;
+
+	for (const FString& File : FoundFiles)
+	{
+		const FString JsonFileName = SelectedDirectory / File;
+		const FString ImageFileName = JsonFileName.Replace(TEXT(".json"), TEXT(".png"));
+
+		// Load the PNG image file for this row into an FImage
+		FImage RowImage;
+		if (IFileManager::Get().FileExists(*ImageFileName))
+		{
+			FImageUtils::LoadImage(*ImageFileName, RowImage);
+		}
+
+		// Open the Json file for reading, and initialize a JsonReader to parse the contents
+		if (TUniquePtr<FArchive> FileReader = TUniquePtr<FArchive>(IFileManager::Get().CreateFileReader(*JsonFileName)))
+		{
+			TSharedRef< TJsonReader<> > JsonReader = TJsonReaderFactory<>::Create(FileReader.Get());
+
+			// Deserialize the row data from the Json file into a Json object
+			TSharedPtr<FJsonObject> JsonRowData = MakeShared<FJsonObject>();
+			if (FJsonSerializer::Deserialize(JsonReader, JsonRowData))
+			{
+				int32 RowNum = NodalOffsetAlgo->ImportCalibrationRow(JsonRowData.ToSharedRef(), RowImage);
+				MaxRowIndex = FMath::Max(MaxRowIndex, RowNum);
+			}
+			else
+			{
+				UE_LOG(LogCameraCalibrationEditor, Verbose, TEXT("Nodal Offset Tool failed to deserialize the dataset row file: %s"), *JsonFileName);
+			}
+		}
+	}
+
+	NodalOffsetAlgo->PostImportCalibrationData();
+
+	// Set the current session's start date/time and row index to match what was just imported to support adding/deleting rows
+	SessionInfo.bIsActive = true;
+	SessionInfo.StartTime = ImportedSessionDateTime;
+	SessionInfo.RowIndex = MaxRowIndex;
+}
+
 #undef LOCTEXT_NAMESPACE
diff --git a/Engine/Plugins/VirtualProduction/CameraCalibration/Source/CameraCalibrationEditor/Private/AssetEditor/NodalOffsetTool.h b/Engine/Plugins/VirtualProduction/CameraCalibration/Source/CameraCalibrationEditor/Private/AssetEditor/NodalOffsetTool.h
--- a/Engine/Plugins/VirtualProduction/CameraCalibration/Source/CameraCalibrationEditor/Private/AssetEditor/NodalOffsetTool.h
+++ b/Engine/Plugins/VirtualProduction/CameraCalibration/Source/CameraCalibrationEditor/Private/AssetEditor/NodalOffsetTool.h
@@ -4,8 +4,9 @@
 
 #include "CoreMinimal.h"
 
-#include "Engine/World.h"
 #include "CameraCalibrationStep.h"
+#include "Engine/World.h"
+#include "ImageCore.h"
 
 #include "NodalOffsetTool.generated.h"
 
@@ -13,8 +14,22 @@ struct FGeometry;
 struct FLensFileEvalData;
 struct FPointerEvent;
 
+class FJsonObject;
 class UCameraNodalOffsetAlgo;
 
+/** Data associated with a nodal offset calibration session */
+struct FNodalOffsetSessionInfo
+{
+	/** The date/time when the current calibration session started */
+	FDateTime StartTime;
+
+	/** The index of the next row in the current calibration session */
+	int32 RowIndex = -1;
+
+	/** True if a calibration session is currently in progress */
+	bool bIsActive = false;
+};
+
 /**
  * FNodalOffsetTool is the controller for the nodal offset tool panel.
  * It has the logic to bridge user input like selection of nodal offset algorithm or CG camera
@@ -45,19 +60,47 @@ public:
 	virtual bool IsOverlayEnabled() const override;
 	//~ End UCameraCalibrationStep interface
 
-public:
-
 	/** Selects the nodal offset algorithm by name */
 	void SetNodalOffsetAlgo(const FName& AlgoName);
 
 	/** Returns the currently selected nodal offset algorithm */
 	UCameraNodalOffsetAlgo* GetNodalOffsetAlgo() const;
 
-public:
-
 	/** Called by the UI when the user wants to save the nodal offset that the current algorithm is providing */
 	void OnSaveCurrentNodalOffset();
 
+	/** Initiate a new calibration session (if one is not already active) */
+	void StartCalibrationSession();
+
+	/** End the current calibration session */
+	void EndCalibrationSession();
+
+	/** Increments the session index and returns its new value */
+	uint32 AdvanceSessionRowIndex();
+
+	/** Intended to be called by one of the algo objects. Exports the data needed by the algo for the current session to a .json file on disk. */
+	void ExportSessionData(const TSharedRef<FJsonObject>& SessionDataObject);
+
+	/** Intended to be called by one of the algo objects. Exports the data associated with a single calibration row to a .json file on disk. */
+	void ExportCalibrationRow(int32 RowIndex, const TSharedRef<FJsonObject>& RowObject, const FImageView& RowImage = FImageView());
+
+	/** Intended to be called by one of the algo objects. Deletes the .json file with the input row index that was previously exported for this session. */
+	void DeleteExportedRow(const int32& RowIndex) const;
+
+	/** Import a saved calibration data set from disk */
+	void ImportCalibrationDataset();
+
+private:
+	/** Get the directory for the current session */
+	FString GetSessionSaveDir() const;
+
+	/** Get the filename of a row with the input index */
+	FString GetRowFilename(int32 RowIndex) const;
+
+public:
+	/** Stores info about the current calibration session of this tool */
+	FNodalOffsetSessionInfo SessionInfo;
+
 private:
 
 	/** Pointer to the calibration steps controller */
diff --git a/Engine/Plugins/VirtualProduction/CameraCalibration/Source/CameraCalibrationEditor/Private/AssetEditor/SLensDistortionToolPanel.cpp b/Engine/Plugins/VirtualProduction/CameraCalibration/Source/CameraCalibrationEditor/Private/AssetEditor/SLensDistortionToolPanel.cpp
--- a/Engine/Plugins/VirtualProduction/CameraCalibration/Source/CameraCalibrationEditor/Private/AssetEditor/SLensDistortionToolPanel.cpp
+++ b/Engine/Plugins/VirtualProduction/CameraCalibration/Source/CameraCalibrationEditor/Private/AssetEditor/SLensDistortionToolPanel.cpp
@@ -3,6 +3,7 @@
 #include "SLensDistortionToolPanel.h"
 
 #include "AssetRegistry/AssetData.h"
+#include "CameraCalibrationSettings.h"
 #include "CameraCalibrationSubsystem.h"
 #include "CameraLensDistortionAlgo.h"
 #include "Dialog/SCustomDialog.h"
@@ -48,6 +49,52 @@ void SLensDistortionToolPanel::Construct(const FArguments& InArgs, ULensDistorti
 			.AutoHeight()
 			[ BuildUIWrapper() ]
 
+			+ SVerticalBox::Slot() // Import dataset
+			.AutoHeight()
+			.HAlign(HAlign_Center)
+			.Padding(0, 20)
+			[
+				SNew(SHorizontalBox)
+
+				+ SHorizontalBox::Slot()
+				.AutoWidth()
+				[
+					SNew(SButton).Text(LOCTEXT("ImportCalibrationDataset", "Import Lens Distortion Dataset"))
+					.HAlign(HAlign_Center)
+					.VAlign(VAlign_Center)
+					.OnClicked_Lambda([this]() -> FReply
+					{
+						if (ULensDistortionTool* LensDistortionTool = Tool.Get())
+						{
+							LensDistortionTool->ImportCalibrationDataset();
+
+							// After importing, the tool may have switched the active algo, so redraw the UI accordingly
+							UCameraLensDistortionAlgo* Algo = LensDistortionTool->GetAlgo();
+							UI->ClearChildren();
+							UI->AddSlot()[Algo->BuildUI()];
+
+							for (const TSharedPtr<FString>& AlgoString : CurrentAlgos)
+							{
+								if (AlgoString->Equals(Algo->FriendlyName().ToString()))
+								{
+									AlgosComboBox->SetSelectedItem(AlgoString);
+									break;
+								}
+							}
+						}
+						return FReply::Handled();
+					})
+					.Visibility_Lambda([]() -> EVisibility
+					{
+						if (GetDefault<UCameraCalibrationSettings>()->IsCalibrationDatasetImportExportEnabled())
+						{
+							return EVisibility::Visible;
+						}
+						return EVisibility::Collapsed;
+					})
+				]
+			]
+
 			+ SVerticalBox::Slot() // Save Offset
 			.AutoHeight()
 			.Padding(0, 20)
@@ -55,11 +102,11 @@ void SLensDistortionToolPanel::Construct(const FArguments& InArgs, ULensDistorti
 				SNew(SButton).Text(LOCTEXT("AddToLUT", "Add To Lens Distortion Calibration"))
 				.HAlign(HAlign_Center)
 				.VAlign(VAlign_Center)
-				.OnClicked_Lambda([&]() -> FReply
+				.OnClicked_Lambda([WeakTool = Tool]() -> FReply
 				{
-					if (Tool.IsValid())
+					if (WeakTool.IsValid())
 					{
-						Tool->OnSaveCurrentCalibrationData();
+						WeakTool->OnSaveCurrentCalibrationData();
 					}
 					return FReply::Handled();
 				})
diff --git a/Engine/Plugins/VirtualProduction/CameraCalibration/Source/CameraCalibrationEditor/Private/AssetEditor/SNodalOffsetToolPanel.cpp b/Engine/Plugins/VirtualProduction/CameraCalibration/Source/CameraCalibrationEditor/Private/AssetEditor/SNodalOffsetToolPanel.cpp
--- a/Engine/Plugins/VirtualProduction/CameraCalibration/Source/CameraCalibrationEditor/Private/AssetEditor/SNodalOffsetToolPanel.cpp
+++ b/Engine/Plugins/VirtualProduction/CameraCalibration/Source/CameraCalibrationEditor/Private/AssetEditor/SNodalOffsetToolPanel.cpp
@@ -3,6 +3,7 @@
 #include "SNodalOffsetToolPanel.h"
 
 #include "AssetRegistry/AssetData.h"
+#include "CameraCalibrationSettings.h"
 #include "CameraCalibrationSubsystem.h"
 #include "CameraNodalOffsetAlgo.h"
 #include "Dialog/SCustomDialog.h"
@@ -49,6 +50,52 @@ void SNodalOffsetToolPanel::Construct(const FArguments& InArgs, UNodalOffsetTool
 			.AutoHeight()
 			[ BuildNodalOffsetUIWrapper() ]
 
+			+ SVerticalBox::Slot() // Import dataset
+			.AutoHeight()
+			.HAlign(HAlign_Center)
+			.Padding(0, 20)
+			[
+				SNew(SHorizontalBox)
+
+				+ SHorizontalBox::Slot()
+				.AutoWidth()
+				[
+					SNew(SButton).Text(LOCTEXT("ImportCalibrationDataset", "Import Nodal Offset Dataset"))
+					.HAlign(HAlign_Center)
+					.VAlign(VAlign_Center)
+					.OnClicked_Lambda([this]() -> FReply
+					{
+						if (UNodalOffsetTool* Tool = NodalOffsetTool.Get())
+						{
+							Tool->ImportCalibrationDataset();
+
+							// After importing, the tool may have switched the active algo, so redraw the UI accordingly
+							UCameraNodalOffsetAlgo* Algo = Tool->GetNodalOffsetAlgo();
+							NodalOffsetUI->ClearChildren();
+							NodalOffsetUI->AddSlot()[Algo->BuildUI()];
+
+							for (const TSharedPtr<FString>& AlgoString : CurrentAlgos)
+							{
+								if (AlgoString->Equals(Algo->FriendlyName().ToString()))
+								{
+									AlgosComboBox->SetSelectedItem(AlgoString);
+									break;
+								}
+							}
+						}
+						return FReply::Handled();
+					})
+					.Visibility_Lambda([]() -> EVisibility
+					{
+						if (GetDefault<UCameraCalibrationSettings>()->IsCalibrationDatasetImportExportEnabled())
+						{
+							return EVisibility::Visible;
+						}
+						return EVisibility::Collapsed;
+					})
+				]
+			]
+
 			+ SVerticalBox::Slot() // Save Offset
 			.AutoHeight()
 			.Padding(0, 20)
@@ -56,11 +103,11 @@ void SNodalOffsetToolPanel::Construct(const FArguments& InArgs, UNodalOffsetTool
 				SNew(SButton).Text(LOCTEXT("AddToNodalOffsetLUT", "Add To Nodal Offset Calibration"))
 				.HAlign(HAlign_Center)
 				.VAlign(VAlign_Center)
-				.OnClicked_Lambda([&]() -> FReply
+				.OnClicked_Lambda([WeakTool = NodalOffsetTool]() -> FReply
 				{
-					if (NodalOffsetTool.IsValid())
+					if (UNodalOffsetTool* Tool = WeakTool.Get())
 					{
-						NodalOffsetTool->OnSaveCurrentNodalOffset();
+						Tool->OnSaveCurrentNodalOffset();
 					}
 					return FReply::Handled();
 				})
diff --git a/Engine/Plugins/VirtualProduction/CameraCalibration/Source/CameraCalibrationEditor/Private/Calibrators/CameraLensDistortionAlgoCheckerboard.cpp b/Engine/Plugins/VirtualProduction/CameraCalibration/Source/CameraCalibrationEditor/Private/Calibrators/CameraLensDistortionAlgoCheckerboard.cpp
--- a/Engine/Plugins/VirtualProduction/CameraCalibration/Source/CameraCalibrationEditor/Private/Calibrators/CameraLensDistortionAlgoCheckerboard.cpp
+++ b/Engine/Plugins/VirtualProduction/CameraCalibration/Source/CameraCalibrationEditor/Private/Calibrators/CameraLensDistortionAlgoCheckerboard.cpp
@@ -6,9 +6,13 @@
 #include "AssetEditor/LensDistortionTool.h"
 #include "AssetEditor/SSimulcamViewport.h"
 #include "CameraCalibrationCheckerboard.h"
+#include "CameraCalibrationEditorLog.h"
+#include "Dom/JsonObject.h"
 #include "EditorFontGlyphs.h"
 #include "Engine/Texture2D.h"
 #include "EngineUtils.h"
+#include "ImageUtils.h"
+#include "JsonObjectConverter.h"
 #include "Misc/MessageDialog.h"
 #include "PropertyCustomizationHelpers.h"
 #include "SphericalLensDistortionModelHandler.h"
@@ -19,19 +23,25 @@
 
 #define LOCTEXT_NAMESPACE "CameraLensDistortionAlgoCheckerboard"
 
+const int UCameraLensDistortionAlgoCheckerboard::DATASET_VERSION = 1;
+
+// String constants for import/export
+namespace UE::CameraCalibration::Private::LensDistortionCheckerboardExportFields
+{
+	static const FString Version(TEXT("Version"));
+}
 
 namespace CameraLensDistortionAlgoCheckerboard
 {
 	class SCalibrationRowGenerator
-		: public SMultiColumnTableRow<TSharedPtr<UCameraLensDistortionAlgoCheckerboard::FCalibrationRowData>>
+		: public SMultiColumnTableRow<TSharedPtr<FLensDistortionCheckerboardRowData>>
 	{
-		using FCalibrationRowData = UCameraLensDistortionAlgoCheckerboard::FCalibrationRowData;
 
 	public:
 		SLATE_BEGIN_ARGS(SCalibrationRowGenerator) {}
 
 		/** The list item for this row */
-		SLATE_ARGUMENT(TSharedPtr<FCalibrationRowData>, CalibrationRowData)
+		SLATE_ARGUMENT(TSharedPtr<FLensDistortionCheckerboardRowData>, CalibrationRowData)
 
 		SLATE_END_ARGS()
 
@@ -39,7 +49,7 @@ namespace CameraLensDistortionAlgoCheckerboard
 		{
 			CalibrationRowData = Args._CalibrationRowData;
 
-			SMultiColumnTableRow<TSharedPtr<FCalibrationRowData>>::Construct(
+			SMultiColumnTableRow<TSharedPtr<FLensDistortionCheckerboardRowData>>::Construct(
 				FSuperRowType::FArguments()
 				.Padding(1.0f),
 				OwnerTableView
@@ -86,7 +96,7 @@ namespace CameraLensDistortionAlgoCheckerboard
 
 
 	private:
-		TSharedPtr<FCalibrationRowData> CalibrationRowData;
+		TSharedPtr<FLensDistortionCheckerboardRowData> CalibrationRowData;
 	};
 
 };
@@ -154,7 +164,8 @@ void UCameraLensDistortionAlgoCheckerboard::Tick(float DeltaTime)
 				break;
 			}
 
-			LastCameraData.LensFileEvalData = *LensFileEvalData;
+			LastCameraData.InputFocus = LensFileEvalData->Input.Focus;
+			LastCameraData.InputZoom = LensFileEvalData->Input.Zoom;
 			LastCameraData.bIsValid = true;
 
 		} while (0);
@@ -252,12 +263,17 @@ bool UCameraLensDistortionAlgoCheckerboard::AddCalibrationRow(FText& OutErrorMes
 	cv::Mat CvGray;
 	cv::cvtColor(CvFrame, CvGray, cv::COLOR_RGBA2GRAY);
 
+	// Export the latest session data
+	ExportSessionData();
+
 	// Create the row that we're going to add
-	TSharedPtr<FCalibrationRowData> Row = MakeShared<FCalibrationRowData>();
+	TSharedPtr<FLensDistortionCheckerboardRowData> Row = MakeShared<FLensDistortionCheckerboardRowData>();
 
-	Row->CameraData = LastCameraData;
+	// Get the next row index for the current calibration session to assign to this new row
+	const uint32 RowIndex = Tool->AdvanceSessionRowIndex();
 
-	Row->Index = CalibrationRows.Num();
+	Row->Index = RowIndex;
+	Row->CameraData = LastCameraData;
 	Row->ImageHeight = CvGray.rows;
 	Row->ImageWidth = CvGray.cols;
 
@@ -311,6 +327,10 @@ bool UCameraLensDistortionAlgoCheckerboard::AddCalibrationRow(FText& OutErrorMes
 			}
 		}
 
+		// Save an image view of the captured frame with the corners overlaid on it (for exporting)
+		FImageView ImageView = FImageView(CvFrame.data, CvFrame.cols, CvFrame.rows, ERawImageFormat::BGRA8);
+		Row->ImageView = ImageView;
+
 		// Update the coverage overlay image with information from the newly added row
 		cv::drawChessboardCorners(CvCoverage, CheckerboardSize, Corners, true);
 
@@ -320,11 +340,12 @@ bool UCameraLensDistortionAlgoCheckerboard::AddCalibrationRow(FText& OutErrorMes
 		// Show the detection to the user
 		if (bShouldShowDetectionWindow)
 		{
-			// Update the coverage overlay image with information from the newly added row
-			cv::drawChessboardCorners(CvFrame, CheckerboardSize, Corners, true);
+			cv::Mat CvFrameWithOverlay;
+			CvFrame.copyTo(CvFrameWithOverlay);
+			cv::drawChessboardCorners(CvFrameWithOverlay, CheckerboardSize, Corners, true);
 
 			FCameraCalibrationWidgetHelpers::DisplayTextureInWindowAlmostFullScreen(
-				FOpenCVHelper::TextureFromCvMat(CvFrame),
+				FOpenCVHelper::TextureFromCvMat(CvFrameWithOverlay),
 				LOCTEXT("CheckerboardDetection", "Checkerboard Detection")
 			);
 		}
@@ -355,6 +376,9 @@ bool UCameraLensDistortionAlgoCheckerboard::AddCalibrationRow(FText& OutErrorMes
 	// Add this data point
 	CalibrationRows.Add(Row);
 
+	// Export the data for this row to a .json file on disk
+	ExportRow(Row);
+
 	// Notify the ListView of the new data
 	if (CalibrationListView.IsValid())
 	{
@@ -403,7 +427,7 @@ TSharedRef<SWidget> UCameraLensDistortionAlgoCheckerboard::BuildUI()
 		[ BuildCalibrationActionButtons() ];
 }
 
-bool UCameraLensDistortionAlgoCheckerboard::ValidateNewRow(TSharedPtr<FCalibrationRowData>& Row, FText& OutErrorMessage) const
+bool UCameraLensDistortionAlgoCheckerboard::ValidateNewRow(TSharedPtr<FLensDistortionCheckerboardRowData>& Row, FText& OutErrorMessage) const
 {
 	if (!Row.IsValid())
 	{
@@ -467,7 +491,7 @@ bool UCameraLensDistortionAlgoCheckerboard::ValidateNewRow(TSharedPtr<FCalibrati
 		return false;
 	}
 
-	const TSharedPtr<FCalibrationRowData>& FirstRow = CalibrationRows[0];
+	const TSharedPtr<FLensDistortionCheckerboardRowData>& FirstRow = CalibrationRows[0];
 
 	// NumRows didn't change
 	if (Row->NumCornerRows != FirstRow->NumCornerRows)
@@ -524,7 +548,7 @@ bool UCameraLensDistortionAlgoCheckerboard::GetLensDistortion(
 	}
 
 	// All points are valid
-	for (const TSharedPtr<FCalibrationRowData>& Row : CalibrationRows)
+	for (const TSharedPtr<FLensDistortionCheckerboardRowData>& Row : CalibrationRows)
 	{
 		if (!ensure(Row.IsValid()))
 		{
@@ -552,7 +576,7 @@ bool UCameraLensDistortionAlgoCheckerboard::GetLensDistortion(
 		return false;
 	}
 
-	const TSharedPtr<FCalibrationRowData>& LastRow = CalibrationRows.Last();
+	const TSharedPtr<FLensDistortionCheckerboardRowData>& LastRow = CalibrationRows.Last();
 
 	// Only parameters data mode supported at the moment
 	if (LensFile->DataMode != ELensDataMode::Parameters)
@@ -582,7 +606,7 @@ bool UCameraLensDistortionAlgoCheckerboard::GetLensDistortion(
 	std::vector<std::vector<cv::Point2f>> Samples2d;
 	std::vector<std::vector<cv::Point3f>> Samples3d;
 
-	for (const TSharedPtr<FCalibrationRowData>& Row : CalibrationRows)
+	for (const TSharedPtr<FLensDistortionCheckerboardRowData>& Row : CalibrationRows)
 	{
 		// add 2d (image) points
 		{
@@ -627,8 +651,8 @@ bool UCameraLensDistortionAlgoCheckerboard::GetLensDistortion(
 	checkSlow(LastRow->ImageHeight > 0);
 
 	// FZ inputs to LUT
-	OutFocus = LastRow->CameraData.LensFileEvalData.Input.Focus;
-	OutZoom = LastRow->CameraData.LensFileEvalData.Input.Zoom;
+	OutFocus = LastRow->CameraData.InputFocus;
+	OutZoom = LastRow->CameraData.InputZoom;
 
 	// FocalLengthInfo
 	OutFocalLengthInfo.FxFy = FVector2D(
@@ -788,10 +812,10 @@ TSharedRef<SWidget> UCameraLensDistortionAlgoCheckerboard::BuildCalibrationActio
 
 TSharedRef<SWidget> UCameraLensDistortionAlgoCheckerboard::BuildCalibrationPointsTable()
 {
-	CalibrationListView = SNew(SListView<TSharedPtr<FCalibrationRowData>>)
+	CalibrationListView = SNew(SListView<TSharedPtr<FLensDistortionCheckerboardRowData>>)
 		.ItemHeight(24)
 		.ListItemsSource(&CalibrationRows)
-		.OnGenerateRow_Lambda([&](TSharedPtr<FCalibrationRowData> InItem, const TSharedRef<STableViewBase>& OwnerTable) -> TSharedRef<ITableRow>
+		.OnGenerateRow_Lambda([&](TSharedPtr<FLensDistortionCheckerboardRowData> InItem, const TSharedRef<STableViewBase>& OwnerTable) -> TSharedRef<ITableRow>
 		{
 			return SNew(CameraLensDistortionAlgoCheckerboard::SCalibrationRowGenerator, OwnerTable)
 				.CalibrationRowData(InItem);
@@ -808,11 +832,17 @@ TSharedRef<SWidget> UCameraLensDistortionAlgoCheckerboard::BuildCalibrationPoint
 			{
 				// Delete selected items
 
-				const TArray<TSharedPtr<FCalibrationRowData>> SelectedItems = CalibrationListView->GetSelectedItems();
+				const TArray<TSharedPtr<FLensDistortionCheckerboardRowData>> SelectedItems = CalibrationListView->GetSelectedItems();
 
-				for (const TSharedPtr<FCalibrationRowData>& SelectedItem : SelectedItems)
+				for (const TSharedPtr<FLensDistortionCheckerboardRowData>& SelectedItem : SelectedItems)
 				{
 					CalibrationRows.Remove(SelectedItem);
+
+					// Delete the previously exported .json file for the row that is being deleted
+					if (ULensDistortionTool* LensDistortionTool = Tool.Get())
+					{
+						LensDistortionTool->DeleteExportedRow(SelectedItem->Index);
+					}
 				}
 
 				CalibrationListView->RequestListRefresh();
@@ -912,6 +942,12 @@ void UCameraLensDistortionAlgoCheckerboard::ClearCalibrationRows()
 	}
 
 	RefreshCoverage();
+
+	// End the current calibration session (a new one will begin the next time a new row is added)
+	if (ULensDistortionTool* LenDistortionTool = Tool.Get())
+	{
+		LenDistortionTool->EndCalibrationSession();
+	}
 }
 
 void UCameraLensDistortionAlgoCheckerboard::RefreshCoverage()
@@ -929,7 +965,7 @@ void UCameraLensDistortionAlgoCheckerboard::RefreshCoverage()
 	CvCoverage.release();
 	CvCoverage = cv::Mat(cv::Size(Size.X, Size.Y), CV_8UC4);
 
-	for (const TSharedPtr<FCalibrationRowData>& Row : CalibrationRows)
+	for (const TSharedPtr<FLensDistortionCheckerboardRowData>& Row : CalibrationRows)
 	{
 		cv::Size CheckerboardSize(Row->NumCornerCols, Row->NumCornerRows);
 
@@ -949,6 +985,92 @@ void UCameraLensDistortionAlgoCheckerboard::RefreshCoverage()
 #endif
 }
 
+
+void UCameraLensDistortionAlgoCheckerboard::ExportSessionData()
+{
+	using namespace UE::CameraCalibration::Private;
+
+	if (ULensDistortionTool* LensDistortionTool = Tool.Get())
+	{
+		// Add all data to a json object that is needed to run this algorithm that is NOT part of a specific row
+		TSharedPtr<FJsonObject> JsonSessionData = MakeShared<FJsonObject>();
+
+		JsonSessionData->SetNumberField(LensDistortionCheckerboardExportFields::Version, DATASET_VERSION);
+
+		// Export the session data to a .json file
+		Tool->ExportSessionData(JsonSessionData.ToSharedRef());
+	}
+}
+
+void UCameraLensDistortionAlgoCheckerboard::ExportRow(TSharedPtr<FLensDistortionCheckerboardRowData> Row)
+{
+	if (ULensDistortionTool* LensDistortionTool = Tool.Get())
+	{
+		if (const TSharedPtr<FJsonObject>& RowObject = FJsonObjectConverter::UStructToJsonObject<FLensDistortionCheckerboardRowData>(Row.ToSharedRef().Get()))
+		{
+			// Export the row data to a .json file
+			LensDistortionTool->ExportCalibrationRow(Row->Index, RowObject.ToSharedRef(), Row->ImageView);
+		}
+	}
+}
+
+bool UCameraLensDistortionAlgoCheckerboard::HasCalibrationData() const
+{
+	return (CalibrationRows.Num() > 0);
+}
+
+void UCameraLensDistortionAlgoCheckerboard::PreImportCalibrationData()
+{
+	// Clear the current set of rows before importing new ones
+	CalibrationRows.Empty();
+}
+
+int32 UCameraLensDistortionAlgoCheckerboard::ImportCalibrationRow(const TSharedRef<FJsonObject>& CalibrationRowObject, const FImage& RowImage)
+{
+	// Create a new row to populate with data from the Json object
+	TSharedPtr<FLensDistortionCheckerboardRowData> NewRow = MakeShared<FLensDistortionCheckerboardRowData>();
+
+	// We enforce strict mode to ensure that every field in the UStruct of row data is present in the imported json.
+	// If any fields are missing, it is likely the row will be invalid, which will lead to errors in the calibration.
+	constexpr int64 CheckFlags = 0;
+	constexpr int64 SkipFlags = 0;
+	constexpr bool bStrictMode = true;
+	if (FJsonObjectConverter::JsonObjectToUStruct<FLensDistortionCheckerboardRowData>(CalibrationRowObject, NewRow.Get(), CheckFlags, SkipFlags, bStrictMode))
+	{
+		if (!RowImage.RawData.IsEmpty())
+		{
+			NewRow->ImageView = RowImage;
+
+			// Resize the image into a thumbnail, create a UTexture from it, and initialize the thumbnail viewport for this row
+			FImage RowThumbnail;
+			RowImage.ResizeTo(RowThumbnail, RowImage.GetWidth() / 4, RowImage.GetHeight() / 4, ERawImageFormat::BGRA8, RowImage.GetGammaSpace());
+
+			UTexture2D* ThumbnailTexture = FImageUtils::CreateTexture2DFromImage(RowThumbnail);
+			NewRow->Thumbnail = SNew(SSimulcamViewport, ThumbnailTexture).WithZoom(false).WithPan(false);
+		}
+
+		CalibrationRows.Add(NewRow);
+	}
+	else
+	{
+		UE_LOG(LogCameraCalibrationEditor, Warning, TEXT("LensDistortionCheckerboard algo failed to import calibration row because at least one field could not be deserialized from the json file."));
+	}
+
+	return NewRow->Index;
+}
+
+void UCameraLensDistortionAlgoCheckerboard::PostImportCalibrationData()
+{
+	// Sort imported calibration rows by row index
+	CalibrationRows.Sort([](const TSharedPtr<FLensDistortionCheckerboardRowData>& LHS, const TSharedPtr<FLensDistortionCheckerboardRowData>& RHS) { return LHS->Index < RHS->Index; });
+
+	// Notify the ListView of the new data
+	if (CalibrationListView.IsValid())
+	{
+		CalibrationListView->RequestListRefresh();
+	}
+}
+
 TSharedRef<SWidget> UCameraLensDistortionAlgoCheckerboard::BuildHelpWidget()
 {
 	return SNew(STextBlock)
diff --git a/Engine/Plugins/VirtualProduction/CameraCalibration/Source/CameraCalibrationEditor/Private/Calibrators/CameraLensDistortionAlgoCheckerboard.h b/Engine/Plugins/VirtualProduction/CameraCalibration/Source/CameraCalibrationEditor/Private/Calibrators/CameraLensDistortionAlgoCheckerboard.h
--- a/Engine/Plugins/VirtualProduction/CameraCalibration/Source/CameraCalibrationEditor/Private/Calibrators/CameraLensDistortionAlgoCheckerboard.h
+++ b/Engine/Plugins/VirtualProduction/CameraCalibration/Source/CameraCalibrationEditor/Private/Calibrators/CameraLensDistortionAlgoCheckerboard.h
@@ -35,6 +35,7 @@ class SListView;
 class SWidget;
 class UMediaTexture;
 class UWorld;
+class FJsonObject;
 
 template<typename OptionType>
 class SComboBox;
@@ -44,6 +45,77 @@ namespace CameraLensDistortionAlgoCheckerboard
 	class SCalibrationRowGenerator;
 }
 
+/** Holds camera information that can be used to add the samples */
+USTRUCT()
+struct FLensDistortionCheckerboardCameraData
+{
+	GENERATED_BODY()
+
+	// True if the rest of the contents are valid.
+	UPROPERTY()
+	bool bIsValid = false;
+
+	// Input focus from lens file evaluation data
+	UPROPERTY()
+	float InputFocus = 0.0f;
+
+	// Input zoom from lens file evaluation data
+	UPROPERTY()
+	float InputZoom = 0.0f;
+};
+
+/** Holds information of the calibration row */
+USTRUCT()
+struct FLensDistortionCheckerboardRowData
+{
+	GENERATED_BODY()
+
+	// Thumbnail to display in list
+	TSharedPtr<SSimulcamViewport> Thumbnail;
+
+	// Captured image containing checkerboard
+	FImageView ImageView;
+
+	// Index to display in list
+	UPROPERTY()
+	int32 Index = -1;
+
+	// Checkerboard corners in 2d image pixel coordinates.
+	UPROPERTY()
+	TArray<FVector2D> Points2d;
+
+	// Checkerboard corners in 3d local space.
+	UPROPERTY()
+	TArray<FVector> Points3d;
+
+	// Which calibrator was used
+	UPROPERTY()
+	FString CalibratorName;
+
+	// Number of corner rows in pattern
+	UPROPERTY()
+	int32 NumCornerRows = 0;
+
+	// Number of corner columns in pattern
+	UPROPERTY()
+	int32 NumCornerCols = 0;
+
+	// Side dimension in cm
+	UPROPERTY()
+	float SquareSideInCm = 0;
+
+	// Width of image
+	UPROPERTY()
+	int32 ImageWidth = 0;
+
+	// Height of image
+	UPROPERTY()
+	int32 ImageHeight = 0;
+
+	// Holds information of the camera data for this sample
+	UPROPERTY()
+	FLensDistortionCheckerboardCameraData CameraData;
+};
 
 /** 
  * Implements a lens distortion calibration algorithm. It requires a checkerboard pattern
@@ -62,6 +134,7 @@ public:
 	virtual bool OnViewportClicked(const FGeometry& MyGeometry, const FPointerEvent& MouseEvent) override;
 	virtual TSharedRef<SWidget> BuildUI() override;
 	virtual FName FriendlyName() const override { return TEXT("Lens Distortion Checkerboard"); };
+	virtual FName ShortName() const override { return TEXT("Checkerboard"); };
 	virtual UMaterialInterface* GetOverlayMaterial() const override;
 	virtual bool IsOverlayEnabled() const override { return bShouldShowOverlay; };
 	virtual void OnDistortionSavedToLens() override;
@@ -75,6 +148,10 @@ public:
 		double& OutError,
 		FText& OutErrorMessage
 	) override;
+	virtual bool HasCalibrationData() const override;
+	virtual void PreImportCalibrationData() override;
+	virtual int32 ImportCalibrationRow(const TSharedRef<FJsonObject>& CalibrationRowObject, const FImage& RowImage) override;
+	virtual void PostImportCalibrationData() override;
 	virtual TSharedRef<SWidget> BuildHelpWidget() override;
 	//~ End CameraLensDistortionAlgo
 
@@ -83,54 +160,9 @@ private:
 	// SCalibrationRowGenerator will need access to the row structures below.
 	friend class CameraLensDistortionAlgoCheckerboard::SCalibrationRowGenerator;
 
-	/** Holds camera information that can be used to add the samples */
-	struct FCameraDataCache
-	{
-		// True if the rest of the contents are valid.
-		bool bIsValid = false;
-
-		// The data used to evaluate the lens data in the camera for this sample
-		FLensFileEvalData LensFileEvalData;
-	};
-
-	/** Holds information of the calibration row */
-	struct FCalibrationRowData
-	{
-		// Thumbnail to display in list
-		TSharedPtr<SSimulcamViewport> Thumbnail;
-
-		// Index to display in list
-		int32 Index = -1;
-
-		// Checkerboard corners in 2d image pixel coordinates.
-		TArray<FVector2D> Points2d;
-
-		// Checkerboard corners in 3d local space.
-		TArray<FVector> Points3d;
-
-		// Which calibrator was used
-		FString CalibratorName;
-
-		// Number of corner rows in pattern
-		int32 NumCornerRows = 0;
-
-		// Number of corner columns in pattern
-		int32 NumCornerCols = 0;
-
-		// Side dimension in cm
-		float SquareSideInCm = 0;
-
-		// Width of image
-		int32 ImageWidth = 0;
-
-		// Height of image
-		int32 ImageHeight = 0;
-
-		// Holds information of the camera data for this sample
-		FCameraDataCache CameraData;
-	};
-
 private:
+	/** Version of the current calibration dataset */
+	static const int DATASET_VERSION;
 
 	/** The lens distortion tool controller */
 	TWeakObjectPtr<ULensDistortionTool> Tool;
@@ -139,13 +171,13 @@ private:
 	TWeakObjectPtr<ACameraCalibrationCheckerboard> Calibrator;
 
 	/** Rows source for the CalibrationListView */
-	TArray<TSharedPtr<FCalibrationRowData>> CalibrationRows;
+	TArray<TSharedPtr<FLensDistortionCheckerboardRowData>> CalibrationRows;
 
 	/** Displays the list of calibration points that will be used to calculate the lens distortion */
-	TSharedPtr<SListView<TSharedPtr<FCalibrationRowData>>> CalibrationListView;
+	TSharedPtr<SListView<TSharedPtr<FLensDistortionCheckerboardRowData>>> CalibrationListView;
 
 	/** Caches the last camera data.  Will hold last value before the media is paused */
-	FCameraDataCache LastCameraData;
+	FLensDistortionCheckerboardCameraData LastCameraData;
 
 	/** True if the coverage overlay should be shown */
 	bool bShouldShowOverlay = false;
@@ -193,7 +225,7 @@ private:
 	void ClearCalibrationRows();
 
 	/** Validates a new calibration point to determine if it should be added as a new sample row */
-	bool ValidateNewRow(TSharedPtr<FCalibrationRowData>& Row, FText& OutErrorMessage) const;
+	bool ValidateNewRow(TSharedPtr<FLensDistortionCheckerboardRowData>& Row, FText& OutErrorMessage) const;
 
 	/** Add a new calibration row from media texture and camera data */
 	bool AddCalibrationRow(FText& OutErrorMessage);
@@ -203,4 +235,10 @@ private:
 
 	/** Update the coverage texture to reflect the current set of calibration rows */
 	void RefreshCoverage();
+
+	/** Export global session data to a .json file */
+	void ExportSessionData();
+
+	/** Export the row data to a json file */
+	void ExportRow(TSharedPtr<FLensDistortionCheckerboardRowData> Row);
 };
diff --git a/Engine/Plugins/VirtualProduction/CameraCalibration/Source/CameraCalibrationEditor/Private/Calibrators/CameraLensDistortionAlgoPoints.cpp b/Engine/Plugins/VirtualProduction/CameraCalibration/Source/CameraCalibrationEditor/Private/Calibrators/CameraLensDistortionAlgoPoints.cpp
--- a/Engine/Plugins/VirtualProduction/CameraCalibration/Source/CameraCalibrationEditor/Private/Calibrators/CameraLensDistortionAlgoPoints.cpp
+++ b/Engine/Plugins/VirtualProduction/CameraCalibration/Source/CameraCalibrationEditor/Private/Calibrators/CameraLensDistortionAlgoPoints.cpp
@@ -6,8 +6,11 @@
 #include "AssetEditor/LensDistortionTool.h"
 #include "CalibrationPointComponent.h"
 #include "Camera/CameraActor.h"
+#include "CameraCalibrationEditorLog.h"
+#include "Dom/JsonObject.h"
 #include "Editor.h"
 #include "Input/Events.h"
+#include "JsonObjectConverter.h"
 #include "Math/Vector.h"
 #include "Misc/MessageDialog.h"
 #include "Models/SphericalLensModel.h"
@@ -30,17 +33,25 @@
 
 #define LOCTEXT_NAMESPACE "CameraLensDistortionAlgoPoints"
 
+const int UCameraLensDistortionAlgoPoints::DATASET_VERSION = 1;
+
+// String constants for import/export
+namespace UE::CameraCalibration::Private::LensDistortionPointsExportFields
+{
+	static const FString Version(TEXT("Version"));
+	static const FString CurrentFocalLength(TEXT("CurrentFocalLength"));
+}
+
 namespace CameraLensDistortionAlgoPoints
 {
-	class SCalibrationRowGenerator : public SMultiColumnTableRow<TSharedPtr<UCameraLensDistortionAlgoPoints::FCalibrationRowData>>
+	class SCalibrationRowGenerator : public SMultiColumnTableRow<TSharedPtr<FLensDistortionPointsRowData>>
 	{
-		using FCalibrationRowData = UCameraLensDistortionAlgoPoints::FCalibrationRowData;
 
 	public:
 		SLATE_BEGIN_ARGS(SCalibrationRowGenerator) {}
 
 		/** The list item for this row */
-		SLATE_ARGUMENT(TSharedPtr<FCalibrationRowData>, CalibrationRowData)
+		SLATE_ARGUMENT(TSharedPtr<FLensDistortionPointsRowData>, CalibrationRowData)
 
 		SLATE_END_ARGS()
 
@@ -48,7 +59,7 @@ namespace CameraLensDistortionAlgoPoints
 		{
 			CalibrationRowData = Args._CalibrationRowData;
 
-			SMultiColumnTableRow<TSharedPtr<FCalibrationRowData>>::Construct(
+			SMultiColumnTableRow<TSharedPtr<FLensDistortionPointsRowData>>::Construct(
 				FSuperRowType::FArguments()
 				.Padding(1.0f),
 				OwnerTableView
@@ -95,7 +106,7 @@ namespace CameraLensDistortionAlgoPoints
 
 
 	private:
-		TSharedPtr<FCalibrationRowData> CalibrationRowData;
+		TSharedPtr<FLensDistortionPointsRowData> CalibrationRowData;
 	};
 };
 
@@ -141,7 +152,7 @@ void UCameraLensDistortionAlgoPoints::Tick(float DeltaTime)
 					continue;
 				}
 
-				FCalibratorPointCache PointCache;
+				FLensDistortionPointsCalibratorPointData PointCache;
 
 				if (!GetCalibratorPointCacheFromName(CalibratorPoint->Name, PointCache))
 				{
@@ -173,7 +184,8 @@ void UCameraLensDistortionAlgoPoints::Tick(float DeltaTime)
 				break;
 			}
 
-			LastCameraData.LensFileEvalData = *LensFileEvalData;
+			LastCameraData.InputFocus = LensFileEvalData->Input.Focus;
+			LastCameraData.InputZoom = LensFileEvalData->Input.Zoom;
 
 			if (!Calibrator.IsValid())
 			{
@@ -209,7 +221,7 @@ bool UCameraLensDistortionAlgoPoints::OnViewportClicked(const FGeometry& MyGeome
 	}
 
 	// Get currently selected calibrator point
-	FCalibratorPointCache LastCalibratorPoint;
+	FLensDistortionPointsCalibratorPointData LastCalibratorPoint;
 	LastCalibratorPoint.bIsValid = false;
 	{
 		if (!CurrentCalibratorPoint.IsValid())
@@ -218,7 +230,7 @@ bool UCameraLensDistortionAlgoPoints::OnViewportClicked(const FGeometry& MyGeome
 		}
 
 		// Find its values in the cache
-		for (const FCalibratorPointCache& PointCache : LastCalibratorPoints)
+		for (const FLensDistortionPointsCalibratorPointData& PointCache : LastCalibratorPoints)
 		{
 			if (PointCache.bIsValid && (PointCache.Name == CurrentCalibratorPoint->Name))
 			{
@@ -240,7 +252,16 @@ bool UCameraLensDistortionAlgoPoints::OnViewportClicked(const FGeometry& MyGeome
 		return true;
 	}
 
-	TSharedPtr<FCalibrationRowData> Row = MakeShared<FCalibrationRowData>();
+	// Export the latest session data
+	ExportSessionData();
+
+	// Create the row that we're going to add
+	TSharedPtr<FLensDistortionPointsRowData> Row = MakeShared<FLensDistortionPointsRowData>();
+
+	// Get the next row index for the current calibration session to assign to this new row
+	const uint32 RowIndex = LensDistortionTool->AdvanceSessionRowIndex();
+
+	Row->Index = RowIndex;
 	Row->CameraData = LastCameraData;
 	Row->CalibratorPointData = LastCalibratorPoint;
 	Row->PatternIndex = CurrentPatternIndex;
@@ -260,6 +281,9 @@ bool UCameraLensDistortionAlgoPoints::OnViewportClicked(const FGeometry& MyGeome
 	// Add this data point
 	CalibrationRows.Add(Row);
 
+	// Export the data for this row to a .json file on disk
+	ExportRow(Row);
+
 	// Notify the ListView of the new data
 	if (CalibrationListView.IsValid())
 	{
@@ -366,7 +390,7 @@ bool UCameraLensDistortionAlgoPoints::GetLensDistortion(
 
 		while (RowIndex < CalibrationRows.Num() && CalibrationRows[RowIndex]->PatternIndex == NextPatternIndex)
 		{
-			const TSharedPtr<FCalibrationRowData>& Row = CalibrationRows[RowIndex];
+			const TSharedPtr<FLensDistortionPointsRowData>& Row = CalibrationRows[RowIndex];
 
 			Points3d.push_back(cv::Point3f(
 				Row->CalibratorPointData.Point3d.X,
@@ -416,13 +440,6 @@ bool UCameraLensDistortionAlgoPoints::GetLensDistortion(
 		return false;
 	}
 
-	if (FMath::IsNearlyEqual(CurrentFocalLengthMM, 0.0f) || CurrentFocalLengthMM < 0.0f)
-	{
-		OutErrorMessage = LOCTEXT("InvalidFocalLengthMsg", "The current focal length (in mm) must be set to a "
-			"valid value in order to seed the calibration algorithm with a best guess as to the camera intrinsics (Fx/Fy)");
-		return false;
-	}
-
 	// Validate sensor dimensions
 	constexpr float MinimumSize = 1.0f; // Limit sensor dimension to 1mm
 	if (LensFile->LensInfo.SensorDimensions.X < MinimumSize || LensFile->LensInfo.SensorDimensions.Y < MinimumSize)
@@ -467,11 +484,11 @@ bool UCameraLensDistortionAlgoPoints::GetLensDistortion(
 	check(DistortionCoefficients.total() == 5);
 	check((CameraMatrix.rows == 3) && (CameraMatrix.cols == 3));
 
-	const TSharedPtr<FCalibrationRowData>& FirstRow = CalibrationRows[0];
+	const TSharedPtr<FLensDistortionPointsRowData>& FirstRow = CalibrationRows[0];
 
 	// FZ inputs to LUT
-	OutFocus = FirstRow->CameraData.LensFileEvalData.Input.Focus;
-	OutZoom = FirstRow->CameraData.LensFileEvalData.Input.Zoom;
+	OutFocus = FirstRow->CameraData.InputFocus;
+	OutZoom = FirstRow->CameraData.InputZoom;
 
 	// FocalLengthInfo
 	OutFocalLengthInfo.FxFy = FVector2D(
@@ -608,7 +625,7 @@ bool UCameraLensDistortionAlgoPoints::AdvanceCalibratorPoint()
 	return false;
 }
 
-bool UCameraLensDistortionAlgoPoints::GetCalibratorPointCacheFromName(const FString& Name, FCalibratorPointCache& CalibratorPointCache) const
+bool UCameraLensDistortionAlgoPoints::GetCalibratorPointCacheFromName(const FString& Name, FLensDistortionPointsCalibratorPointData& CalibratorPointCache) const
 {
 	CalibratorPointCache.bIsValid = false;
 
@@ -633,7 +650,7 @@ bool UCameraLensDistortionAlgoPoints::GetCalibratorPointCacheFromName(const FStr
 	return false;
 }
 
-bool UCameraLensDistortionAlgoPoints::ValidateNewRow(TSharedPtr<FCalibrationRowData>& Row, FText& OutErrorMessage) const
+bool UCameraLensDistortionAlgoPoints::ValidateNewRow(TSharedPtr<FLensDistortionPointsRowData>& Row, FText& OutErrorMessage) const
 {
 	const FCameraCalibrationStepsController* StepsController = LensDistortionTool->GetCameraCalibrationStepsController();
 
@@ -657,7 +674,7 @@ bool UCameraLensDistortionAlgoPoints::ValidateNewRow(TSharedPtr<FCalibrationRowD
 	}
 
 	// Same camera as before
-	const TSharedPtr<FCalibrationRowData>& FirstRow = CalibrationRows[0];
+	const TSharedPtr<FLensDistortionPointsRowData>& FirstRow = CalibrationRows[0];
 
 	if (FirstRow->CameraData.UniqueId != Row->CameraData.UniqueId)
 	{
@@ -701,15 +718,21 @@ void UCameraLensDistortionAlgoPoints::ClearCalibrationRows()
 	{
 		CurrentCalibratorPoint = nullptr;
 	}
+
+	// End the current calibration session (a new one will begin the next time a new row is added)
+	if (ULensDistortionTool* Tool = LensDistortionTool.Get())
+	{
+		Tool->EndCalibrationSession();
+	}
 }
 
 void UCameraLensDistortionAlgoPoints::DeleteSelectedCalibrationRows()
 {
-	const TArray<TSharedPtr<FCalibrationRowData>> SelectedItems = CalibrationListView->GetSelectedItems();
+	const TArray<TSharedPtr<FLensDistortionPointsRowData>> SelectedItems = CalibrationListView->GetSelectedItems();
 
 	TSet<uint32> PatternsToDelete;
 
-	for (const TSharedPtr<FCalibrationRowData>& Row : SelectedItems)
+	for (const TSharedPtr<FLensDistortionPointsRowData>& Row : SelectedItems)
 	{
 		PatternsToDelete.Add(Row->PatternIndex);
 	}
@@ -726,21 +749,27 @@ void UCameraLensDistortionAlgoPoints::DeleteSelectedCalibrationRows()
 		}
 	}
 
-	for (const TSharedPtr<FCalibrationRowData>& SelectedItem : SelectedItems)
+	for (const TSharedPtr<FLensDistortionPointsRowData>& SelectedItem : SelectedItems)
 	{
 		CalibrationRows.Remove(SelectedItem);
+
+		// Delete the previously exported .json file for the row that is being deleted
+		if (ULensDistortionTool* Tool = LensDistortionTool.Get())
+		{
+			Tool->DeleteExportedRow(SelectedItem->Index);
+		}
 	}
 
 	CalibrationListView->RequestListRefresh();
 }
 
-void UCameraLensDistortionAlgoPoints::OnCalibrationRowSelectionChanged(TSharedPtr<FCalibrationRowData> SelectedRow, ESelectInfo::Type SelectInfo)
+void UCameraLensDistortionAlgoPoints::OnCalibrationRowSelectionChanged(TSharedPtr<FLensDistortionPointsRowData> SelectedRow, ESelectInfo::Type SelectInfo)
 {
 	if (SelectInfo != ESelectInfo::Type::Direct)
 	{
 		TSet<uint32> PatternsToSelect;
 
-		for (const TSharedPtr<FCalibrationRowData>& Row : CalibrationRows)
+		for (const TSharedPtr<FLensDistortionPointsRowData>& Row : CalibrationRows)
 		{
 			if (CalibrationListView->IsItemSelected(Row))
 			{
@@ -748,10 +777,10 @@ void UCameraLensDistortionAlgoPoints::OnCalibrationRowSelectionChanged(TSharedPt
 			}
 		}
 
-		TArray<TSharedPtr<FCalibrationRowData>> ItemsToSelect;
+		TArray<TSharedPtr<FLensDistortionPointsRowData>> ItemsToSelect;
 
 		// Select all of the rows that have a matching pattern index (so that whole patterns are selected as groups)
-		for (const TSharedPtr<FCalibrationRowData>& Row : CalibrationRows)
+		for (const TSharedPtr<FLensDistortionPointsRowData>& Row : CalibrationRows)
 		{
 			if (PatternsToSelect.Contains(Row->PatternIndex))
 			{
@@ -820,20 +849,25 @@ TSharedRef<SWidget> UCameraLensDistortionAlgoPoints::BuildCurrentFocalLengthWidg
 				InValue = 1.0f;
 			}
 			CurrentFocalLengthMM = InValue;
+		})
+		.OnValueCommitted_Lambda([&](float InValue, ETextCommit::Type CommitType)
+		{
+			// Re-export the session data to update the new focal length value
+			ExportSessionData();
 		});
 }
 
 TSharedRef<SWidget> UCameraLensDistortionAlgoPoints::BuildCalibrationPointsTable()
 {
-	CalibrationListView = SNew(SListView<TSharedPtr<FCalibrationRowData>>)
+	CalibrationListView = SNew(SListView<TSharedPtr<FLensDistortionPointsRowData>>)
 		.ItemHeight(24)
 		.ListItemsSource(&CalibrationRows)
-		.OnGenerateRow_Lambda([&](TSharedPtr<FCalibrationRowData> InItem, const TSharedRef<STableViewBase>& OwnerTable) -> TSharedRef<ITableRow>
+		.OnGenerateRow_Lambda([&](TSharedPtr<FLensDistortionPointsRowData> InItem, const TSharedRef<STableViewBase>& OwnerTable) -> TSharedRef<ITableRow>
 		{
 			return SNew(CameraLensDistortionAlgoPoints::SCalibrationRowGenerator, OwnerTable).CalibrationRowData(InItem);
 		})
 		.SelectionMode(ESelectionMode::Multi)
-		.OnSelectionChanged_Lambda([&](TSharedPtr<FCalibrationRowData> SelectedRow, ESelectInfo::Type SelectInfo)
+		.OnSelectionChanged_Lambda([&](TSharedPtr<FLensDistortionPointsRowData> SelectedRow, ESelectInfo::Type SelectInfo)
 		{
 			OnCalibrationRowSelectionChanged(SelectedRow, SelectInfo);
 		})
@@ -908,6 +942,94 @@ TSharedRef<SWidget> UCameraLensDistortionAlgoPoints::BuildCalibrationActionButto
 	;
 }
 
+void UCameraLensDistortionAlgoPoints::ExportSessionData()
+{
+	using namespace UE::CameraCalibration::Private;
+	if (ULensDistortionTool* Tool = LensDistortionTool.Get())
+	{
+		// Add all data to a json object that is needed to run this algorithm that is NOT part of a specific row
+		TSharedPtr<FJsonObject> JsonSessionData = MakeShared<FJsonObject>();
+
+		JsonSessionData->SetNumberField(LensDistortionPointsExportFields::CurrentFocalLength, CurrentFocalLengthMM);
+		JsonSessionData->SetNumberField(LensDistortionPointsExportFields::Version, DATASET_VERSION);
+
+		// Export the session data to a .json file
+		Tool->ExportSessionData(JsonSessionData.ToSharedRef());
+	}
+}
+
+void UCameraLensDistortionAlgoPoints::ExportRow(TSharedPtr<FLensDistortionPointsRowData> Row)
+{
+	if (ULensDistortionTool* Tool = LensDistortionTool.Get())
+	{
+		if (const TSharedPtr<FJsonObject>& RowObject = FJsonObjectConverter::UStructToJsonObject<FLensDistortionPointsRowData>(Row.ToSharedRef().Get()))
+		{
+			// Export the row data to a .json file
+			LensDistortionTool->ExportCalibrationRow(Row->Index, RowObject.ToSharedRef());
+		}
+	}
+}
+
+bool UCameraLensDistortionAlgoPoints::HasCalibrationData() const
+{
+	return (CalibrationRows.Num() > 0);
+}
+
+void UCameraLensDistortionAlgoPoints::PreImportCalibrationData()
+{
+	// Clear the current set of rows before importing new ones
+	CalibrationRows.Empty();
+}
+
+void UCameraLensDistortionAlgoPoints::ImportSessionData(const TSharedRef<FJsonObject>& SessionDataObject)
+{
+	using namespace UE::CameraCalibration::Private;
+	SessionDataObject->TryGetNumberField(LensDistortionPointsExportFields::CurrentFocalLength, CurrentFocalLengthMM);
+}
+
+int32 UCameraLensDistortionAlgoPoints::ImportCalibrationRow(const TSharedRef<FJsonObject>& CalibrationRowObject, const FImage& RowImage)
+{
+	// Create a new row to populate with data from the Json object
+	TSharedPtr<FLensDistortionPointsRowData> NewRow = MakeShared<FLensDistortionPointsRowData>();
+
+	// We enforce strict mode to ensure that every field in the UStruct of row data is present in the imported json.
+	// If any fields are missing, it is likely the row will be invalid, which will lead to errors in the calibration.
+	constexpr int64 CheckFlags = 0;
+	constexpr int64 SkipFlags = 0;
+	constexpr bool bStrictMode = true;
+	if (FJsonObjectConverter::JsonObjectToUStruct<FLensDistortionPointsRowData>(CalibrationRowObject, NewRow.Get(), CheckFlags, SkipFlags, bStrictMode))
+	{
+		CalibrationRows.Add(NewRow);
+	}
+	else
+	{
+		UE_LOG(LogCameraCalibrationEditor, Warning, TEXT("LensDistortionPoints algo failed to import calibration row because at least one field could not be deserialized from the json file."));
+	}
+
+	return NewRow->Index;
+}
+
+void UCameraLensDistortionAlgoPoints::PostImportCalibrationData()
+{
+	// Find the maximum pattern index among the imported rows and set the current index to be greater than that
+	uint32 MaxPatternIndex = 0;
+	for (const TSharedPtr<FLensDistortionPointsRowData>& Row : CalibrationRows)
+	{
+		MaxPatternIndex = FMath::Max(MaxPatternIndex, Row->PatternIndex);
+	}
+
+	CurrentPatternIndex = MaxPatternIndex + 1;
+
+	// Sort imported calibration rows by row index
+	CalibrationRows.Sort([](const TSharedPtr<FLensDistortionPointsRowData>& LHS, const TSharedPtr<FLensDistortionPointsRowData>& RHS) { return LHS->Index < RHS->Index; });
+
+	// Notify the ListView of the new data
+	if (CalibrationListView.IsValid())
+	{
+		CalibrationListView->RequestListRefresh();
+	}
+}
+
 TSharedRef<SWidget> UCameraLensDistortionAlgoPoints::BuildHelpWidget()
 {
 	return SNew(STextBlock)
diff --git a/Engine/Plugins/VirtualProduction/CameraCalibration/Source/CameraCalibrationEditor/Private/Calibrators/CameraLensDistortionAlgoPoints.h b/Engine/Plugins/VirtualProduction/CameraCalibration/Source/CameraCalibrationEditor/Private/Calibrators/CameraLensDistortionAlgoPoints.h
--- a/Engine/Plugins/VirtualProduction/CameraCalibration/Source/CameraCalibrationEditor/Private/Calibrators/CameraLensDistortionAlgoPoints.h
+++ b/Engine/Plugins/VirtualProduction/CameraCalibration/Source/CameraCalibrationEditor/Private/Calibrators/CameraLensDistortionAlgoPoints.h
@@ -11,11 +11,86 @@
 template <typename ItemType>
 class SListView;
 
+class FJsonObject;
+
 namespace CameraLensDistortionAlgoPoints
 {
 	class SCalibrationRowGenerator;
 };
 
+/** Holds information of the identified calibrator 2d point for a given sample of a 2d-3d correlation */
+USTRUCT()
+struct FLensDistortionPointsCalibratorPointData
+{
+	GENERATED_BODY()
+
+	// True if the rest of the contents are valid.
+	UPROPERTY()
+	bool bIsValid = false;
+
+	// The name of the 3d calibrator point, as defined in the mesh
+	UPROPERTY()
+	FString Name;
+
+	// 3D location of the calibrator point
+	UPROPERTY()
+	FVector Point3d = FVector::ZeroVector;
+
+	// 2D location of the calibrator point
+	UPROPERTY()
+	FVector2D Point2d = FVector2D::ZeroVector;
+};
+
+/** Holds information of the camera pose for a given sample of a 2d-3d correlation */
+USTRUCT()
+struct FLensDistortionPointsCameraData
+{
+	GENERATED_BODY()
+
+	// True if the rest of the contents are valid.
+	UPROPERTY()
+	bool bIsValid = false;
+
+	// The unique id of the camera object. Used to detect camera selection changes during a calibration session.
+	UPROPERTY()
+	uint32 UniqueId = INDEX_NONE;
+
+	// The unique id of the calibrator object. Used to detect calibrator selection changes during a calibration session.
+	UPROPERTY()
+	uint32 CalibratorUniqueId = INDEX_NONE;
+
+	// Input focus from lens file evaluation data
+	UPROPERTY()
+	float InputFocus = 0.0f;
+
+	// Input zoom from lens file evaluation data
+	UPROPERTY()
+	float InputZoom = 0.0f;
+};
+
+/** Holds information of the calibrator 3d point for a given sample of a 2d-3d correlation */
+USTRUCT()
+struct FLensDistortionPointsRowData
+{
+	GENERATED_BODY()
+	
+	// Index to display in list
+	UPROPERTY()
+	int32 Index = INDEX_NONE;
+
+	// The 3D-2D point correspondences for each of the calibrator points in the chosen calibrator for a single pose
+	UPROPERTY()
+	FLensDistortionPointsCalibratorPointData CalibratorPointData;
+
+	// Holds information of the camera data for this sample
+	UPROPERTY()
+	FLensDistortionPointsCameraData CameraData;
+
+	// Index of the calibration pattern this row is associated with
+	UPROPERTY()
+	uint32 PatternIndex = INDEX_NONE;
+};
+
 /**
  * Implements a lens distortion calibration algorithm. It uses 3d points (UCalibrationPointComponent)
  * specified in a selectable calibrator with features that the user can identify by clicking on the
@@ -40,6 +115,8 @@ public:
 	virtual bool OnViewportClicked(const FGeometry& MyGeometry, const FPointerEvent& MouseEvent) override;
 	virtual TSharedRef<SWidget> BuildUI() override;
 	virtual FName FriendlyName() const override { return TEXT("Lens Distortion Points Method"); };
+	virtual FName ShortName() const override { return TEXT("Points"); };
+	virtual void OnDistortionSavedToLens() override;
 
 	virtual bool GetLensDistortion(
 		float& OutFocus,
@@ -51,7 +128,11 @@ public:
 		double& OutError,
 		FText& OutErrorMessage) override;
 
-	virtual void OnDistortionSavedToLens() override;
+	virtual bool HasCalibrationData() const override;
+	virtual void PreImportCalibrationData() override;
+	virtual void ImportSessionData(const TSharedRef<FJsonObject>& SessionDataObject) override;
+	virtual int32 ImportCalibrationRow(const TSharedRef<FJsonObject>& CalibrationRowObject, const FImage& RowImage) override;
+	virtual void PostImportCalibrationData() override;
 	virtual TSharedRef<SWidget> BuildHelpWidget() override;
 	//~ End UCameraLensDistortionAlgo
 
@@ -71,52 +152,9 @@ protected:
 		FString Name;
 	};
 
-	/** Holds information of the identified calibrator 2d point for a given sample of a 2d-3d correlation */
-	struct FCalibratorPointCache
-	{
-		// True if the rest of the contents are valid.
-		bool bIsValid;
-
-		// The name of the 3d calibrator point, as defined in the mesh
-		FString Name;
-
-		// 3D location of the calibrator point
-		FVector Point3d;
-
-		// 2D location of the calibrator point
-		FVector2D Point2d;
-	};
-
-	/** Holds information of the camera pose for a given sample of a 2d-3d correlation */
-	struct FCameraDataCache
-	{
-		// True if the rest of the contents are valid.
-		bool bIsValid;
-
-		// The unique id of the camera object. Used to detect camera selection changes during a calibration session.
-		uint32 UniqueId;
-
-		// The unique id of the calibrator object. Used to detect calibrator selection changes during a calibration session.
-		uint32 CalibratorUniqueId;
-
-		// The data used to evaluate the lens data in the camera for this sample
-		FLensFileEvalData LensFileEvalData;
-	};
-
-	/** Holds information of the calibrator 3d point for a given sample of a 2d-3d correlation */
-	struct FCalibrationRowData
-	{
-		// The 3D-2D point correspondences for each of the calibrator points in the chosen calibrator for a single pose
-		FCalibratorPointCache CalibratorPointData;
-
-		// Holds information of the camera data for this sample
-		FCameraDataCache CameraData;
-
-		// Index of the calibration pattern this row is associated with
-		uint32 PatternIndex;
-	};
-
 protected:
+	/** Version of the current calibration dataset */
+	static const int DATASET_VERSION;
 
 	/** The lens distortion tool */
 	TWeakObjectPtr<ULensDistortionTool> LensDistortionTool;
@@ -131,16 +169,16 @@ protected:
 	TSharedPtr<FCalibratorPointData> CurrentCalibratorPoint;
 
 	/** Rows source for the CalibrationListView */
-	TArray<TSharedPtr<FCalibrationRowData>> CalibrationRows;
+	TArray<TSharedPtr<FLensDistortionPointsRowData>> CalibrationRows;
 
 	/** Displays the list of calibration points that will be used to calculate the lens distortion */
-	TSharedPtr<SListView<TSharedPtr<FCalibrationRowData>>> CalibrationListView;
+	TSharedPtr<SListView<TSharedPtr<FLensDistortionPointsRowData>>> CalibrationListView;
 
 	/** Caches the last calibrator point 3d location.  Will hold last value before the lens distortion tool is paused */
-	TArray<FCalibratorPointCache> LastCalibratorPoints;
+	TArray<FLensDistortionPointsCalibratorPointData> LastCalibratorPoints;
 
 	/** Caches the last camera data.  Will hold last value before the lens distortion tool is paused */
-	FCameraDataCache LastCameraData;
+	FLensDistortionPointsCameraData LastCameraData;
 
 	/** Instructs the inline allocation policy how many elements to allocate in a single allocation */
 	static constexpr uint32 NumInlineAllocations = 32;
@@ -170,10 +208,10 @@ protected:
 	 *
 	 * @return True if successful.
 	 */
-	bool GetCalibratorPointCacheFromName(const FString& Name, FCalibratorPointCache& CalibratorPointCache) const;
+	bool GetCalibratorPointCacheFromName(const FString& Name, FLensDistortionPointsCalibratorPointData& CalibratorPointCache) const;
 
 	/** Validates a new calibration point to determine if it should be added as a new sample row */
-	bool ValidateNewRow(TSharedPtr<FCalibrationRowData>& Row, FText& OutErrorMessage) const;
+	bool ValidateNewRow(TSharedPtr<FLensDistortionPointsRowData>& Row, FText& OutErrorMessage) const;
 
 	/** Clears the list of calibration sample rows */
 	void ClearCalibrationRows();
@@ -182,7 +220,13 @@ protected:
 	void DeleteSelectedCalibrationRows();
 
 	/** Updates the selection in the list of calibration rows so that patterns are always grouped together */
-	void OnCalibrationRowSelectionChanged(TSharedPtr<FCalibrationRowData> SelectedRow, ESelectInfo::Type SelectInfo);
+	void OnCalibrationRowSelectionChanged(TSharedPtr<FLensDistortionPointsRowData> SelectedRow, ESelectInfo::Type SelectInfo);
+
+	/** Export global session data to a .json file */
+	void ExportSessionData();
+
+	/** Export the row data to a .json file */
+	void ExportRow(TSharedPtr<FLensDistortionPointsRowData> Row);
 
 	/** Builds the UI of the calibration device picker */
 	TSharedRef<SWidget> BuildCalibrationDevicePickerWidget();
diff --git a/Engine/Plugins/VirtualProduction/CameraCalibration/Source/CameraCalibrationEditor/Private/Calibrators/CameraNodalOffsetAlgoAruco.cpp b/Engine/Plugins/VirtualProduction/CameraCalibration/Source/CameraCalibrationEditor/Private/Calibrators/CameraNodalOffsetAlgoAruco.cpp
--- a/Engine/Plugins/VirtualProduction/CameraCalibration/Source/CameraCalibrationEditor/Private/Calibrators/CameraNodalOffsetAlgoAruco.cpp
+++ b/Engine/Plugins/VirtualProduction/CameraCalibration/Source/CameraCalibrationEditor/Private/Calibrators/CameraNodalOffsetAlgoAruco.cpp
@@ -264,23 +264,29 @@ bool UCameraNodalOffsetAlgoAruco::PopulatePoints(FText& OutErrorMessage)
 
 		check(MarkerCorners[MarkerIdx].size() == 4); // Successful detection should always return the 4 corners of each marker.
 
-		// Iterate over the 4 corners of this marker and try to find the corresponding calibration point
+		// Export the latest session data
+		ExportSessionData();
+
+		 // Iterate over the 4 corners of this marker and try to find the corresponding calibration point
 		for (int32 CornerIdx = 0; CornerIdx < 4; ++CornerIdx)
 		{
 			// Build calibrator point name based on the detected marker
 			const FString MarkerName = FString::Printf(TEXT("%s-%d-%s"), *DictionaryInfo.Name, MarkerIds[MarkerIdx], *CornerNames[CornerIdx]); //-V557
 
-			FCalibratorPointCache PointCache;
+			FNodalOffsetPointsCalibratorPointData PointCache;
 
 			if (!CalibratorPointCacheFromName(MarkerName, PointCache))
 			{
 				continue;
 			}
 
-			// Create and populate the new calibration row that we're going to add, corresponding to this marker corner
+			// Create the row that we're going to add
+			TSharedPtr<FNodalOffsetPointsRowData> Row = MakeShared<FNodalOffsetPointsRowData>();
 
-			TSharedPtr<FCalibrationRowData> Row = MakeShared<FCalibrationRowData>();
+			// Get the next row index for the current calibration session to assign to this new row
+			const uint32 RowIndex = NodalOffsetTool->AdvanceSessionRowIndex();
 
+			Row->Index = RowIndex;
 			Row->Point2D.X = float(MarkerCorners[MarkerIdx][CornerIdx].x) / Size.X;
 			Row->Point2D.Y = float(MarkerCorners[MarkerIdx][CornerIdx].y) / Size.Y;
 
@@ -302,6 +308,9 @@ bool UCameraNodalOffsetAlgoAruco::PopulatePoints(FText& OutErrorMessage)
 			}
 
 			CalibrationRows.Add(Row);
+
+			// Export the data for this row to a .json file on disk
+			ExportRow(Row);
 		}
 	}
 
diff --git a/Engine/Plugins/VirtualProduction/CameraCalibration/Source/CameraCalibrationEditor/Private/Calibrators/CameraNodalOffsetAlgoAruco.h b/Engine/Plugins/VirtualProduction/CameraCalibration/Source/CameraCalibrationEditor/Private/Calibrators/CameraNodalOffsetAlgoAruco.h
--- a/Engine/Plugins/VirtualProduction/CameraCalibration/Source/CameraCalibrationEditor/Private/Calibrators/CameraNodalOffsetAlgoAruco.h
+++ b/Engine/Plugins/VirtualProduction/CameraCalibration/Source/CameraCalibrationEditor/Private/Calibrators/CameraNodalOffsetAlgoAruco.h
@@ -23,6 +23,7 @@ public:
 	virtual bool OnViewportClicked(const FGeometry& MyGeometry, const FPointerEvent& MouseEvent) override;
 	virtual TSharedRef<SWidget> BuildUI() override;
 	virtual FName FriendlyName() const override { return TEXT("Nodal Offset Aruco Markers"); };
+	virtual FName ShortName() const override { return TEXT("Aruco"); };
 	virtual TSharedRef<SWidget> BuildHelpWidget() override;
 	//~ End CalibPointsNodalOffsetAlgo
 
diff --git a/Engine/Plugins/VirtualProduction/CameraCalibration/Source/CameraCalibrationEditor/Private/Calibrators/CameraNodalOffsetAlgoCheckerboard.cpp b/Engine/Plugins/VirtualProduction/CameraCalibration/Source/CameraCalibrationEditor/Private/Calibrators/CameraNodalOffsetAlgoCheckerboard.cpp
--- a/Engine/Plugins/VirtualProduction/CameraCalibration/Source/CameraCalibrationEditor/Private/Calibrators/CameraNodalOffsetAlgoCheckerboard.cpp
+++ b/Engine/Plugins/VirtualProduction/CameraCalibration/Source/CameraCalibrationEditor/Private/Calibrators/CameraNodalOffsetAlgoCheckerboard.cpp
@@ -197,6 +197,9 @@ bool UCameraNodalOffsetAlgoCheckerboard::PopulatePoints(FText& OutErrorMessage)
 		std::reverse(Points2d.begin(), Points2d.end());
 	}
 
+	// Export the latest session data
+	ExportSessionData();
+
 	// Create and populate the new calibration rows that we're going to add
 	for (int32 RowIdx = 0; RowIdx < Checkerboard->NumCornerRows; ++RowIdx)
 	{
@@ -204,8 +207,12 @@ bool UCameraNodalOffsetAlgoCheckerboard::PopulatePoints(FText& OutErrorMessage)
 		{
 			const int32 PointIdx = RowIdx * Checkerboard->NumCornerCols + ColIdx;
 
-			TSharedPtr<FCalibrationRowData> Row = MakeShared<FCalibrationRowData>();
+			TSharedPtr<FNodalOffsetPointsRowData> Row = MakeShared<FNodalOffsetPointsRowData>();
+
+			// Get the next row index for the current calibration session to assign to this new row
+			const uint32 RowIndex = NodalOffsetTool->AdvanceSessionRowIndex();
 
+			Row->Index = RowIndex;
 			Row->Point2D.X = float(Points2d[PointIdx].x) / Size.X;
 			Row->Point2D.Y = float(Points2d[PointIdx].y) / Size.Y;
 
@@ -230,6 +237,9 @@ bool UCameraNodalOffsetAlgoCheckerboard::PopulatePoints(FText& OutErrorMessage)
 			}
 
 			CalibrationRows.Add(Row);
+
+			// Export the data for this row to a .json file on disk
+			ExportRow(Row);
 		}
 	}
 
diff --git a/Engine/Plugins/VirtualProduction/CameraCalibration/Source/CameraCalibrationEditor/Private/Calibrators/CameraNodalOffsetAlgoCheckerboard.h b/Engine/Plugins/VirtualProduction/CameraCalibration/Source/CameraCalibrationEditor/Private/Calibrators/CameraNodalOffsetAlgoCheckerboard.h
--- a/Engine/Plugins/VirtualProduction/CameraCalibration/Source/CameraCalibrationEditor/Private/Calibrators/CameraNodalOffsetAlgoCheckerboard.h
+++ b/Engine/Plugins/VirtualProduction/CameraCalibration/Source/CameraCalibrationEditor/Private/Calibrators/CameraNodalOffsetAlgoCheckerboard.h
@@ -24,6 +24,7 @@ public:
 	virtual bool OnViewportClicked(const FGeometry& MyGeometry, const FPointerEvent& MouseEvent) override;
 	virtual TSharedRef<SWidget> BuildUI() override;
 	virtual FName FriendlyName() const override { return TEXT("Nodal Offset Checkerboard"); };
+	virtual FName ShortName() const override { return TEXT("Checkerboard"); };
 	virtual TSharedRef<SWidget> BuildHelpWidget() override;
 	//~ End CalibPointsNodalOffsetAlgo
 
diff --git a/Engine/Plugins/VirtualProduction/CameraCalibration/Source/CameraCalibrationEditor/Private/Calibrators/CameraNodalOffsetAlgoOpticalAxis.cpp b/Engine/Plugins/VirtualProduction/CameraCalibration/Source/CameraCalibrationEditor/Private/Calibrators/CameraNodalOffsetAlgoOpticalAxis.cpp
--- a/Engine/Plugins/VirtualProduction/CameraCalibration/Source/CameraCalibrationEditor/Private/Calibrators/CameraNodalOffsetAlgoOpticalAxis.cpp
+++ b/Engine/Plugins/VirtualProduction/CameraCalibration/Source/CameraCalibrationEditor/Private/Calibrators/CameraNodalOffsetAlgoOpticalAxis.cpp
@@ -208,11 +208,11 @@ bool UCameraNodalOffsetAlgoOpticalAxis::GetNodalOffset(FNodalPointOffset& OutNod
  	}
 
 	// Validate that all rows have approximately the same camera pose
-	const TSharedPtr<FCalibrationRowData>& FirstRow = CalibrationRows[0];
+	const TSharedPtr<FNodalOffsetPointsRowData>& FirstRow = CalibrationRows[0];
 	const FTransform CameraPose = FirstRow->CameraData.Pose;
 	CachedInverseCameraPose = CameraPose.Inverse();
 
-	for (const TSharedPtr<FCalibrationRowData>& Row : CalibrationRows)
+	for (const TSharedPtr<FNodalOffsetPointsRowData>& Row : CalibrationRows)
 	{
 		if (!FCameraCalibrationUtils::IsNearlyEqual(CameraPose, Row->CameraData.Pose))
 		{
@@ -222,15 +222,18 @@ bool UCameraNodalOffsetAlgoOpticalAxis::GetNodalOffset(FNodalPointOffset& OutNod
 	}
 
 	// Cache the evaluated focus and zoom so that modifications to the nodal offset table will update the same point
-	OutFocus = CachedFocus = FirstRow->CameraData.LensFileEvalData.Input.Focus;
-	OutZoom = CachedZoom = FirstRow->CameraData.LensFileEvalData.Input.Zoom;
+	CachedFocus = FirstRow->CameraData.InputFocus;
+	CachedZoom = FirstRow->CameraData.InputZoom;
+
+	OutFocus = CachedFocus;
+	OutZoom = CachedZoom;
 
 #if WITH_OPENCV
 	// Gather the 3D points from each of the calibration rows
 	std::vector<cv::Point3f> Points3d;
 	Points3d.reserve(CalibrationRows.Num());
 
-	for (const TSharedPtr<FCalibrationRowData>& Row : CalibrationRows)
+	for (const TSharedPtr<FNodalOffsetPointsRowData>& Row : CalibrationRows)
 	{
 		Points3d.push_back(cv::Point3f(
 			Row->CalibratorPointData.Location.X,
diff --git a/Engine/Plugins/VirtualProduction/CameraCalibration/Source/CameraCalibrationEditor/Private/Calibrators/CameraNodalOffsetAlgoOpticalAxis.h b/Engine/Plugins/VirtualProduction/CameraCalibration/Source/CameraCalibrationEditor/Private/Calibrators/CameraNodalOffsetAlgoOpticalAxis.h
--- a/Engine/Plugins/VirtualProduction/CameraCalibration/Source/CameraCalibrationEditor/Private/Calibrators/CameraNodalOffsetAlgoOpticalAxis.h
+++ b/Engine/Plugins/VirtualProduction/CameraCalibration/Source/CameraCalibrationEditor/Private/Calibrators/CameraNodalOffsetAlgoOpticalAxis.h
@@ -31,6 +31,7 @@ public:
 	virtual TSharedRef<SWidget> BuildUI() override;
 	virtual bool GetNodalOffset(FNodalPointOffset& OutNodalOffset, float& OutFocus, float& OutZoom, float& OutError, FText& OutErrorMessage) override;
 	virtual FName FriendlyName() const override { return TEXT("Nodal Offset Optical Axis"); };
+	virtual FName ShortName() const override { return TEXT("OpticalAxis"); };
 	virtual UMaterialInterface* GetOverlayMaterial() const override;
 	virtual bool IsOverlayEnabled() const override;
 	virtual TSharedRef<SWidget> BuildHelpWidget() override;
diff --git a/Engine/Plugins/VirtualProduction/CameraCalibration/Source/CameraCalibrationEditor/Private/Calibrators/CameraNodalOffsetAlgoPoints.cpp b/Engine/Plugins/VirtualProduction/CameraCalibration/Source/CameraCalibrationEditor/Private/Calibrators/CameraNodalOffsetAlgoPoints.cpp
--- a/Engine/Plugins/VirtualProduction/CameraCalibration/Source/CameraCalibrationEditor/Private/Calibrators/CameraNodalOffsetAlgoPoints.cpp
+++ b/Engine/Plugins/VirtualProduction/CameraCalibration/Source/CameraCalibrationEditor/Private/Calibrators/CameraNodalOffsetAlgoPoints.cpp
@@ -10,11 +10,13 @@
 #include "CameraCalibrationUtils.h"
 #include "CineCameraComponent.h"
 #include "DistortionRenderingUtils.h"
+#include "Dom/JsonObject.h"
 #include "Editor.h"
 #include "Framework/MultiBox/MultiBoxBuilder.h"
 #include "GameFramework/Actor.h"
 #include "Input/Events.h"
 #include "Internationalization/Text.h"
+#include "JsonObjectConverter.h"
 #include "Layout/Geometry.h"
 #include "LensComponent.h"
 #include "LensDistortionModelHandlerBase.h"
@@ -46,18 +48,25 @@ static TAutoConsoleVariable<float> CVarRotationStepValue(TEXT("CameraCalibration
 static TAutoConsoleVariable<float> CVarLocationStepValue(TEXT("CameraCalibration.LocationStepValue"), 0.5, TEXT("The value of the initial step size to use when finding an optimal nodal offset location that minimizes reprojection error."));
 #endif
 
+const int UCameraNodalOffsetAlgoPoints::DATASET_VERSION = 1;
+
+// String constants for import/export
+namespace UE::CameraCalibration::Private::NodalOffsetPointsExportFields
+{
+	static const FString Version(TEXT("Version"));
+}
+
 namespace CameraNodalOffsetAlgoPoints
 {
 	class SCalibrationRowGenerator
-		: public SMultiColumnTableRow<TSharedPtr<UCameraNodalOffsetAlgoPoints::FCalibrationRowData>>
+		: public SMultiColumnTableRow<TSharedPtr<FNodalOffsetPointsRowData>>
 	{
-		using FCalibrationRowData = UCameraNodalOffsetAlgoPoints::FCalibrationRowData;
 
 	public:
 		SLATE_BEGIN_ARGS(SCalibrationRowGenerator) {}
 
 		/** The list item for this row */
-		SLATE_ARGUMENT(TSharedPtr<FCalibrationRowData>, CalibrationRowData)
+		SLATE_ARGUMENT(TSharedPtr<FNodalOffsetPointsRowData>, CalibrationRowData)
 
 		SLATE_END_ARGS()
 
@@ -65,7 +74,7 @@ namespace CameraNodalOffsetAlgoPoints
 		{
 			CalibrationRowData = Args._CalibrationRowData;
 
-			SMultiColumnTableRow<TSharedPtr<FCalibrationRowData>>::Construct(
+			SMultiColumnTableRow<TSharedPtr<FNodalOffsetPointsRowData>>::Construct(
 				FSuperRowType::FArguments()
 				.Padding(1.0f),
 				OwnerTableView
@@ -106,7 +115,7 @@ namespace CameraNodalOffsetAlgoPoints
 
 
 	private:
-		TSharedPtr<FCalibrationRowData> CalibrationRowData;
+		TSharedPtr<FNodalOffsetPointsRowData> CalibrationRowData;
 	};
 
 	/** Contains basic result of a nodal offset calibration based on a single camera pose for all samples */
@@ -283,7 +292,7 @@ private:
 };
 #endif
 
-double UCameraNodalOffsetAlgoPoints::MinimizeReprojectionError(FTransform& InOutNodalOffset, const TArray<TSharedPtr<TArray<TSharedPtr<FCalibrationRowData>>>>& SamePoseRowGroups) const
+double UCameraNodalOffsetAlgoPoints::MinimizeReprojectionError(FTransform& InOutNodalOffset, const TArray<TSharedPtr<TArray<TSharedPtr<FNodalOffsetPointsRowData>>>>& SamePoseRowGroups) const
 {
 #if WITH_OPENCV
 	cv::Ptr<cv::DownhillSolver> Solver = cv::DownhillSolver::create();
@@ -344,15 +353,15 @@ double UCameraNodalOffsetAlgoPoints::MinimizeReprojectionError(FTransform& InOut
 		return -1.0;
 	}
 
-	const TSharedPtr<FCalibrationRowData>& FirstRow = CalibrationRows[0];
+	const TSharedPtr<FNodalOffsetPointsRowData>& FirstRow = CalibrationRows[0];
 
 	FTransform ExistingOffset = FTransform::Identity;
 
 	if (FirstRow->CameraData.bWasNodalOffsetApplied)
 	{
 		FNodalPointOffset NodalPointOffset;
-		const float Focus = FirstRow->CameraData.LensFileEvalData.Input.Focus;
-		const float Zoom = FirstRow->CameraData.LensFileEvalData.Input.Zoom;
+		const float Focus = FirstRow->CameraData.InputFocus;
+		const float Zoom = FirstRow->CameraData.InputZoom;
 
 		if (LensFile->EvaluateNodalPointOffset(Focus, Zoom, NodalPointOffset))
 		{
@@ -370,14 +379,14 @@ double UCameraNodalOffsetAlgoPoints::MinimizeReprojectionError(FTransform& InOut
 
 	for (int32 PoseGroupIndex = 0; PoseGroupIndex < SamePoseRowGroups.Num(); ++PoseGroupIndex)
 	{
-		const TArray<TSharedPtr<FCalibrationRowData>>& PoseGroup = *(SamePoseRowGroups[PoseGroupIndex]);
+		const TArray<TSharedPtr<FNodalOffsetPointsRowData>>& PoseGroup = *(SamePoseRowGroups[PoseGroupIndex]);
 
 		if (!ensureMsgf((PoseGroup.Num() > 0), TEXT("Not enough calibration rows")))
 		{
 			return -1.0;
 		}
 
-		const TSharedPtr<FCalibrationRowData>& FirstRowInPoseGroup = PoseGroup[0];
+		const TSharedPtr<FNodalOffsetPointsRowData>& FirstRowInPoseGroup = PoseGroup[0];
 
 		FTransform CameraPose = FirstRowInPoseGroup->CameraData.Pose;
 
@@ -386,7 +395,7 @@ double UCameraNodalOffsetAlgoPoints::MinimizeReprojectionError(FTransform& InOut
 		Points3d.reserve(PoseGroup.Num());
 		Points2d.reserve(PoseGroup.Num());
 
-		for (const TSharedPtr<FCalibrationRowData>& Row : PoseGroup)
+		for (const TSharedPtr<FNodalOffsetPointsRowData>& Row : PoseGroup)
 		{
 			// Convert from UE coordinates to OpenCV coordinates
 			FTransform Transform;
@@ -423,7 +432,7 @@ double UCameraNodalOffsetAlgoPoints::MinimizeReprojectionError(FTransform& InOut
 #endif
 }
 
-double UCameraNodalOffsetAlgoPoints::ComputeReprojectionError(const FTransform& NodalOffset, const TArray<TSharedPtr<FCalibrationRowData>>& PoseGroup) const
+double UCameraNodalOffsetAlgoPoints::ComputeReprojectionError(const FTransform& NodalOffset, const TArray<TSharedPtr<FNodalOffsetPointsRowData>>& PoseGroup) const
 {
 #if WITH_OPENCV
 	const FCameraCalibrationStepsController* StepsController;
@@ -450,14 +459,14 @@ double UCameraNodalOffsetAlgoPoints::ComputeReprojectionError(const FTransform&
 		return -1.0;
 	}
 
-	const TSharedPtr<FCalibrationRowData>& FirstRow = PoseGroup[0];
+	const TSharedPtr<FNodalOffsetPointsRowData>& FirstRow = PoseGroup[0];
 
 	if (!ensureMsgf(FirstRow->bUndistortedIsValid, TEXT("This method operates on undistorted 2D points. Call UndistortCalibrationRowPoints() prior to calling this method")))
 	{
 		return -1.0;
 	}
 
-	for (const TSharedPtr<FCalibrationRowData>& Row : PoseGroup)
+	for (const TSharedPtr<FNodalOffsetPointsRowData>& Row : PoseGroup)
 	{
 		// Convert from UE coordinates to OpenCV coordinates
 
@@ -487,8 +496,8 @@ double UCameraNodalOffsetAlgoPoints::ComputeReprojectionError(const FTransform&
 	if (FirstRow->CameraData.bWasNodalOffsetApplied)
 	{
 		FNodalPointOffset NodalPointOffset;
-		const float Focus = FirstRow->CameraData.LensFileEvalData.Input.Focus;
-		const float Zoom = FirstRow->CameraData.LensFileEvalData.Input.Zoom;
+		const float Focus = FirstRow->CameraData.InputFocus;
+		const float Zoom = FirstRow->CameraData.InputZoom;
 
 		if (LensFile->EvaluateNodalPointOffset(Focus, Zoom, NodalPointOffset))
 		{
@@ -568,7 +577,7 @@ void UCameraNodalOffsetAlgoPoints::Tick(float DeltaTime)
 					continue;
 				}
 
-				FCalibratorPointCache PointCache;
+				FNodalOffsetPointsCalibratorPointData PointCache;
 
 				if (!CalibratorPointCacheFromName(CalibratorPoint->Name, PointCache))
 				{
@@ -608,7 +617,8 @@ void UCameraNodalOffsetAlgoPoints::Tick(float DeltaTime)
 
 			LastCameraData.Pose = CameraComponent->GetComponentToWorld();
 			LastCameraData.UniqueId = Camera->GetUniqueID();
-			LastCameraData.LensFileEvalData = *LensFileEvalData;
+			LastCameraData.InputFocus = LensFileEvalData->Input.Focus;
+			LastCameraData.InputZoom = LensFileEvalData->Input.Zoom;
 
 			const ULensComponent* LensComponent = StepsController->FindLensComponent();
 			if (LensComponent)
@@ -692,7 +702,7 @@ bool UCameraNodalOffsetAlgoPoints::OnViewportClicked(const FGeometry& MyGeometry
 	}
 
 	// Get currently selected calibrator point
-	FCalibratorPointCache LastCalibratorPoint;
+	FNodalOffsetPointsCalibratorPointData LastCalibratorPoint;
 	LastCalibratorPoint.bIsValid = false;
 	{
 		TSharedPtr<FCalibratorPointData> CalibratorPoint = CalibratorPointsComboBox->GetSelectedItem();
@@ -703,7 +713,7 @@ bool UCameraNodalOffsetAlgoPoints::OnViewportClicked(const FGeometry& MyGeometry
 		}
 
 		// Find its values in the cache
-		for (const FCalibratorPointCache& PointCache : LastCalibratorPoints)
+		for (const FNodalOffsetPointsCalibratorPointData& PointCache : LastCalibratorPoints)
 		{
 			if (PointCache.bIsValid && (PointCache.Name == CalibratorPoint->Name))
 			{
@@ -719,9 +729,16 @@ bool UCameraNodalOffsetAlgoPoints::OnViewportClicked(const FGeometry& MyGeometry
 		return true;
 	}
 
+	// Export the latest session data
+	ExportSessionData();
+
 	// Create the row that we're going to add
-	TSharedPtr<FCalibrationRowData> Row = MakeShared<FCalibrationRowData>();
+	TSharedPtr<FNodalOffsetPointsRowData> Row = MakeShared<FNodalOffsetPointsRowData>();
+
+	// Get the next row index for the current calibration session to assign to this new row
+	const uint32 RowIndex = NodalOffsetTool->AdvanceSessionRowIndex();
 
+	Row->Index = RowIndex;
 	Row->CameraData = LastCameraData;
 	Row->CalibratorPointData = LastCalibratorPoint;
 
@@ -746,6 +763,9 @@ bool UCameraNodalOffsetAlgoPoints::OnViewportClicked(const FGeometry& MyGeometry
 	// Add this data point
 	CalibrationRows.Add(Row);
 
+	// Export the data for this row to a .json file on disk
+	ExportRow(Row);
+
 	// Notify the ListView of the new data
 	if (CalibrationListView.IsValid())
 	{
@@ -796,7 +816,7 @@ bool UCameraNodalOffsetAlgoPoints::GetCurrentCalibratorPointLocation(FVector& Ou
 		return false;
 	}
 
-	FCalibratorPointCache PointCache;
+	FNodalOffsetPointsCalibratorPointData PointCache;
 
 	if (!CalibratorPointCacheFromName(CalibratorPointData->Name, PointCache))
 	{
@@ -842,7 +862,7 @@ TSharedRef<SWidget> UCameraNodalOffsetAlgoPoints::BuildUI()
 		;
 }
 
-bool UCameraNodalOffsetAlgoPoints::ValidateNewRow(TSharedPtr<FCalibrationRowData>& Row, FText& OutErrorMessage) const
+bool UCameraNodalOffsetAlgoPoints::ValidateNewRow(TSharedPtr<FNodalOffsetPointsRowData>& Row, FText& OutErrorMessage) const
 {
 	const FCameraCalibrationStepsController* StepsController;
 	const ULensFile* LensFile;
@@ -867,14 +887,6 @@ bool UCameraNodalOffsetAlgoPoints::ValidateNewRow(TSharedPtr<FCalibrationRowData
 		return false;
 	}
 
-	// Same LensFile
-
-	if (Row->CameraData.LensFileEvalData.LensFile != LensFile)
-	{
-		OutErrorMessage = LOCTEXT("LensFileWasNotTheSame", "Lens file was not the same");
-		return false;
-	}
-
 	// Same camera in view
 
 	const ACameraActor* Camera = StepsController->GetCamera();
@@ -897,7 +909,7 @@ bool UCameraNodalOffsetAlgoPoints::ValidateNewRow(TSharedPtr<FCalibrationRowData
 	}
 
 	// Same camera as before
-	const TSharedPtr<FCalibrationRowData>& FirstRow = CalibrationRows[0];
+	const TSharedPtr<FNodalOffsetPointsRowData>& FirstRow = CalibrationRows[0];
 
 	if (FirstRow->CameraData.UniqueId != Row->CameraData.UniqueId)
 	{
@@ -931,7 +943,7 @@ bool UCameraNodalOffsetAlgoPoints::ValidateNewRow(TSharedPtr<FCalibrationRowData
 	return true;
 }
 
-bool UCameraNodalOffsetAlgoPoints::BasicCalibrationChecksPass(const TArray<TSharedPtr<FCalibrationRowData>>& Rows, FText& OutErrorMessage) const
+bool UCameraNodalOffsetAlgoPoints::BasicCalibrationChecksPass(const TArray<TSharedPtr<FNodalOffsetPointsRowData>>& Rows, FText& OutErrorMessage) const
 {
 	const FCameraCalibrationStepsController* StepsController;
 	const ULensFile* LensFile;
@@ -953,7 +965,7 @@ bool UCameraNodalOffsetAlgoPoints::BasicCalibrationChecksPass(const TArray<TShar
 	}
 
 	// All points are valid
-	for (const TSharedPtr<FCalibrationRowData>& Row : Rows)
+	for (const TSharedPtr<FNodalOffsetPointsRowData>& Row : Rows)
 	{
 		if (!ensure(Row.IsValid()))
 		{
@@ -983,7 +995,7 @@ bool UCameraNodalOffsetAlgoPoints::BasicCalibrationChecksPass(const TArray<TShar
 		return false;
 	}
 
-	const TSharedPtr<FCalibrationRowData>& FirstRow = Rows[0];
+	const TSharedPtr<FNodalOffsetPointsRowData>& FirstRow = Rows[0];
 
 	// Still same camera (since we need it to get the distortion handler, which much be the same)
 
@@ -994,7 +1006,7 @@ bool UCameraNodalOffsetAlgoPoints::BasicCalibrationChecksPass(const TArray<TShar
 	}
 
 	// Camera did not move much.
-	for (const TSharedPtr<FCalibrationRowData>& Row : Rows)
+	for (const TSharedPtr<FNodalOffsetPointsRowData>& Row : Rows)
 	{
 		if (!FCameraCalibrationUtils::IsNearlyEqual(FirstRow->CameraData.Pose, Row->CameraData.Pose))
 		{
@@ -1008,7 +1020,7 @@ bool UCameraNodalOffsetAlgoPoints::BasicCalibrationChecksPass(const TArray<TShar
 
 bool UCameraNodalOffsetAlgoPoints::CalculatedOptimalCameraComponentPose(
 	FTransform& OutDesiredCameraTransform, 
-	const TArray<TSharedPtr<FCalibrationRowData>>& Rows, 
+	const TArray<TSharedPtr<FNodalOffsetPointsRowData>>& Rows, 
 	FText& OutErrorMessage) const
 {
 	if (!BasicCalibrationChecksPass(Rows, OutErrorMessage))
@@ -1046,7 +1058,7 @@ bool UCameraNodalOffsetAlgoPoints::CalculatedOptimalCameraComponentPose(
 	Points3d.reserve(Rows.Num());
 	Points2d.reserve(Rows.Num());
 
-	for (const TSharedPtr<FCalibrationRowData>& Row : Rows)
+	for (const TSharedPtr<FNodalOffsetPointsRowData>& Row : Rows)
 	{
 		// Convert from UE coordinates to OpenCV coordinates
 		FTransform Transform;
@@ -1156,16 +1168,16 @@ bool UCameraNodalOffsetAlgoPoints::CalculatedOptimalCameraComponentPose(
 }
 
 bool UCameraNodalOffsetAlgoPoints::CalibratorMovedInAnyRow(
-	const TArray<TSharedPtr<FCalibrationRowData>>& Rows) const
+	const TArray<TSharedPtr<FNodalOffsetPointsRowData>>& Rows) const
 {
 	if (!Rows.Num())
 	{
 		return false;
 	}
 
-	TSharedPtr<FCalibrationRowData> FirstRow;
+	TSharedPtr<FNodalOffsetPointsRowData> FirstRow;
 
-	for (const TSharedPtr<FCalibrationRowData>& Row : Rows)
+	for (const TSharedPtr<FNodalOffsetPointsRowData>& Row : Rows)
 	{
 		if (!FirstRow.IsValid())
 		{
@@ -1193,9 +1205,9 @@ bool UCameraNodalOffsetAlgoPoints::CalibratorMovedInAnyRow(
 
 
 bool UCameraNodalOffsetAlgoPoints::CalibratorMovedAcrossGroups(
-	const TArray<TSharedPtr<TArray<TSharedPtr<FCalibrationRowData>>>>& SamePoseRowGroups) const
+	const TArray<TSharedPtr<TArray<TSharedPtr<FNodalOffsetPointsRowData>>>>& SamePoseRowGroups) const
 {
-	TArray<TSharedPtr<FCalibrationRowData>> Rows;
+	TArray<TSharedPtr<FNodalOffsetPointsRowData>> Rows;
 
 	for (const auto& SamePoseRowGroup : SamePoseRowGroups)
 	{
@@ -1215,7 +1227,7 @@ bool UCameraNodalOffsetAlgoPoints::GetNodalOffsetSinglePose(
 	float& OutFocus, 
 	float& OutZoom, 
 	float& OutError, 
-	const TArray<TSharedPtr<FCalibrationRowData>>& Rows, 
+	const TArray<TSharedPtr<FNodalOffsetPointsRowData>>& Rows, 
 	FText& OutErrorMessage) const
 {
 	const FCameraCalibrationStepsController* StepsController;
@@ -1252,10 +1264,10 @@ bool UCameraNodalOffsetAlgoPoints::GetNodalOffsetSinglePose(
 	check(Rows.Num()); // There must have been rows for CalculatedOptimalCameraComponentPose to have succeeded.
 	check(Rows[0].IsValid()); // All rows should be valid.
 
-	const TSharedPtr<FCalibrationRowData>& FirstRow = Rows[0];
+	const TSharedPtr<FNodalOffsetPointsRowData>& FirstRow = Rows[0];
 
-	OutFocus = FirstRow->CameraData.LensFileEvalData.Input.Focus;
-	OutZoom  = FirstRow->CameraData.LensFileEvalData.Input.Zoom;
+	OutFocus = FirstRow->CameraData.InputFocus;
+	OutZoom  = FirstRow->CameraData.InputZoom;
 
 	// See if the camera already had an offset applied, in which case we need to account for it.
 
@@ -1294,7 +1306,7 @@ bool UCameraNodalOffsetAlgoPoints::GetNodalOffset(FNodalPointOffset& OutNodalOff
 	}
 
 	// Group Rows by camera poses.
-	TArray<TSharedPtr<TArray<TSharedPtr<FCalibrationRowData>>>> SamePoseRowGroups;
+	TArray<TSharedPtr<TArray<TSharedPtr<FNodalOffsetPointsRowData>>>> SamePoseRowGroups;
 	GroupRowsByCameraPose(SamePoseRowGroups, CalibrationRows);
 
 	if (!SamePoseRowGroups.Num())
@@ -1648,7 +1660,7 @@ bool UCameraNodalOffsetAlgoPoints::ApplyNodalOffsetToCalibrator()
 
 	// All calibration points should correspond to the same calibrator
 
-	for (const TSharedPtr<FCalibrationRowData>& Row : CalibrationRows)
+	for (const TSharedPtr<FNodalOffsetPointsRowData>& Row : CalibrationRows)
 	{
 		check(Row.IsValid());
 
@@ -1661,7 +1673,7 @@ bool UCameraNodalOffsetAlgoPoints::ApplyNodalOffsetToCalibrator()
 		}
 	}
 
-	const TSharedPtr<FCalibrationRowData>& LastRow = CalibrationRows[CalibrationRows.Num() - 1];
+	const TSharedPtr<FNodalOffsetPointsRowData>& LastRow = CalibrationRows[CalibrationRows.Num() - 1];
 	check(LastRow.IsValid());
 
 	// Verify that calibrator did not move much for all the samples
@@ -1674,7 +1686,7 @@ bool UCameraNodalOffsetAlgoPoints::ApplyNodalOffsetToCalibrator()
 	}
 
 	// Group Rows by camera poses.
-	TArray<TSharedPtr<TArray<TSharedPtr<FCalibrationRowData>>>> SamePoseRowGroups;
+	TArray<TSharedPtr<TArray<TSharedPtr<FNodalOffsetPointsRowData>>>> SamePoseRowGroups;
 	GroupRowsByCameraPose(SamePoseRowGroups, CalibrationRows);
 
 	if (!SamePoseRowGroups.Num())
@@ -1735,7 +1747,7 @@ bool UCameraNodalOffsetAlgoPoints::ApplyNodalOffsetToCalibrator()
 }
 
 bool UCameraNodalOffsetAlgoPoints::CalcTrackingOriginPoseForSingleCamPose(
-	const TArray<TSharedPtr<FCalibrationRowData>>& Rows, 
+	const TArray<TSharedPtr<FNodalOffsetPointsRowData>>& Rows, 
 	FTransform& OutTransform, 
 	FText& OutErrorMessage)
 {
@@ -1755,7 +1767,7 @@ bool UCameraNodalOffsetAlgoPoints::CalcTrackingOriginPoseForSingleCamPose(
 
 	check(Rows.Num()); // Must be non-zero if CalculatedOptimalCameraComponentPose succeeded.
 
-	const TSharedPtr<FCalibrationRowData>& LastRow = Rows[Rows.Num() - 1];
+	const TSharedPtr<FNodalOffsetPointsRowData>& LastRow = Rows[Rows.Num() - 1];
 	check(LastRow.IsValid());
 
 	// calculate the new parent transform
@@ -1826,7 +1838,7 @@ bool UCameraNodalOffsetAlgoPoints::ApplyNodalOffsetToTrackingOrigin()
 		return false;
 	}
 
-	const TSharedPtr<FCalibrationRowData>& LastRow = CalibrationRows[CalibrationRows.Num() - 1];
+	const TSharedPtr<FNodalOffsetPointsRowData>& LastRow = CalibrationRows[CalibrationRows.Num() - 1];
 	check(LastRow.IsValid());
 
 	if (LastRow->CameraData.ParentUniqueId != ParentActor->GetUniqueID())
@@ -1838,7 +1850,7 @@ bool UCameraNodalOffsetAlgoPoints::ApplyNodalOffsetToTrackingOrigin()
 	}
 
 	// Group Rows by camera poses.
-	TArray<TSharedPtr<TArray<TSharedPtr<FCalibrationRowData>>>> SamePoseRowGroups;
+	TArray<TSharedPtr<TArray<TSharedPtr<FNodalOffsetPointsRowData>>>> SamePoseRowGroups;
 	GroupRowsByCameraPose(SamePoseRowGroups, CalibrationRows);
 
 	TArray<FSinglePoseResult> SinglePoseResults;
@@ -1893,7 +1905,7 @@ bool UCameraNodalOffsetAlgoPoints::ApplyNodalOffsetToTrackingOrigin()
 }
 
 bool UCameraNodalOffsetAlgoPoints::CalcCalibratorPoseForSingleCamPose(
-	const TArray<TSharedPtr<FCalibrationRowData>>& Rows,
+	const TArray<TSharedPtr<FNodalOffsetPointsRowData>>& Rows,
 	FTransform& OutTransform,
 	FText& OutErrorMessage)
 {
@@ -1909,7 +1921,7 @@ bool UCameraNodalOffsetAlgoPoints::CalcCalibratorPoseForSingleCamPose(
 
 	check(Rows.Num());
 
-	const TSharedPtr<FCalibrationRowData>& LastRow = Rows[Rows.Num() - 1];
+	const TSharedPtr<FNodalOffsetPointsRowData>& LastRow = Rows[Rows.Num() - 1];
 	check(LastRow.IsValid());
 
 	// Calculate the offset
@@ -1958,7 +1970,7 @@ bool UCameraNodalOffsetAlgoPoints::ApplyNodalOffsetToCalibratorParent()
 
 	// All calibration points should correspond to the same calibrator and calibrator parent
 
-	for (const TSharedPtr<FCalibrationRowData>& Row : CalibrationRows)
+	for (const TSharedPtr<FNodalOffsetPointsRowData>& Row : CalibrationRows)
 	{
 		check(Row.IsValid());
 
@@ -1989,7 +2001,7 @@ bool UCameraNodalOffsetAlgoPoints::ApplyNodalOffsetToCalibratorParent()
 	}
 
 	// Group Rows by camera poses.
-	TArray<TSharedPtr<TArray<TSharedPtr<FCalibrationRowData>>>> SamePoseRowGroups;
+	TArray<TSharedPtr<TArray<TSharedPtr<FNodalOffsetPointsRowData>>>> SamePoseRowGroups;
 	GroupRowsByCameraPose(SamePoseRowGroups, CalibrationRows);
 
 	if (!SamePoseRowGroups.Num())
@@ -2041,7 +2053,7 @@ bool UCameraNodalOffsetAlgoPoints::ApplyNodalOffsetToCalibratorParent()
 		return false;
 	}
 
-	const TSharedPtr<FCalibrationRowData>& LastRow = CalibrationRows[CalibrationRows.Num() - 1];
+	const TSharedPtr<FNodalOffsetPointsRowData>& LastRow = CalibrationRows[CalibrationRows.Num() - 1];
 	check(LastRow.IsValid());
 
 	// Apply the new calibrator parent transform
@@ -2097,7 +2109,7 @@ bool UCameraNodalOffsetAlgoPoints::ApplyNodalOffsetToCalibratorComponents()
 
 	// All calibration points should correspond to the same calibrator and same set of calibrator components
 
-	for (const TSharedPtr<FCalibrationRowData>& Row : CalibrationRows)
+	for (const TSharedPtr<FNodalOffsetPointsRowData>& Row : CalibrationRows)
 	{
 		check(Row.IsValid());
 
@@ -2133,7 +2145,7 @@ bool UCameraNodalOffsetAlgoPoints::ApplyNodalOffsetToCalibratorComponents()
 	}
 
 	// Group Rows by camera poses.
-	TArray<TSharedPtr<TArray<TSharedPtr<FCalibrationRowData>>>> SamePoseRowGroups;
+	TArray<TSharedPtr<TArray<TSharedPtr<FNodalOffsetPointsRowData>>>> SamePoseRowGroups;
 	GroupRowsByCameraPose(SamePoseRowGroups, CalibrationRows);
 
 	if (!SamePoseRowGroups.Num())
@@ -2185,7 +2197,7 @@ bool UCameraNodalOffsetAlgoPoints::ApplyNodalOffsetToCalibratorComponents()
 		return false;
 	}
 
-	const TSharedPtr<FCalibrationRowData>& LastRow = CalibrationRows[CalibrationRows.Num() - 1];
+	const TSharedPtr<FNodalOffsetPointsRowData>& LastRow = CalibrationRows[CalibrationRows.Num() - 1];
 	check(LastRow.IsValid());
 
 	// Apply the desired transform to each of the calibrator components
@@ -2195,8 +2207,11 @@ bool UCameraNodalOffsetAlgoPoints::ApplyNodalOffsetToCalibratorComponents()
 		{
 			if (USceneComponent* AttachComponent = CalibratorComponent->GetAttachParent())
 			{
-				AttachComponent->Modify();
-				AttachComponent->SetWorldTransform(LastRow->CameraData.CalibratorComponentPoses[AttachComponent->GetUniqueID()] * LastRow->CameraData.CalibratorPose.Inverse() * DesiredCalibratorPose);
+				if (LastRow->CameraData.CalibratorComponentPoses.Contains(AttachComponent->GetUniqueID()))
+				{
+					AttachComponent->Modify();
+					AttachComponent->SetWorldTransform(LastRow->CameraData.CalibratorComponentPoses[AttachComponent->GetUniqueID()] * LastRow->CameraData.CalibratorPose.Inverse()* DesiredCalibratorPose);
+				}
 			}
 		}
 	}
@@ -2237,10 +2252,10 @@ TSharedRef<SWidget> UCameraNodalOffsetAlgoPoints::BuildCalibrationPointsComboBox
 
 TSharedRef<SWidget> UCameraNodalOffsetAlgoPoints::BuildCalibrationPointsTable()
 {
-	CalibrationListView = SNew(SListView<TSharedPtr<FCalibrationRowData>>)
+	CalibrationListView = SNew(SListView<TSharedPtr<FNodalOffsetPointsRowData>>)
 		.ItemHeight(24)
 		.ListItemsSource(&CalibrationRows)
-		.OnGenerateRow_Lambda([&](TSharedPtr<FCalibrationRowData> InItem, const TSharedRef<STableViewBase>& OwnerTable) -> TSharedRef<ITableRow>
+		.OnGenerateRow_Lambda([&](TSharedPtr<FNodalOffsetPointsRowData> InItem, const TSharedRef<STableViewBase>& OwnerTable) -> TSharedRef<ITableRow>
 		{
 			return SNew(CameraNodalOffsetAlgoPoints::SCalibrationRowGenerator, OwnerTable).CalibrationRowData(InItem);
 		})
@@ -2256,11 +2271,17 @@ TSharedRef<SWidget> UCameraNodalOffsetAlgoPoints::BuildCalibrationPointsTable()
 			{
 				// Delete selected items
 
-				const TArray<TSharedPtr<FCalibrationRowData>> SelectedItems = CalibrationListView->GetSelectedItems();
+				const TArray<TSharedPtr<FNodalOffsetPointsRowData>> SelectedItems = CalibrationListView->GetSelectedItems();
 
-				for (const TSharedPtr<FCalibrationRowData>& SelectedItem : SelectedItems)
+				for (const TSharedPtr<FNodalOffsetPointsRowData>& SelectedItem : SelectedItems)
 				{
 					CalibrationRows.Remove(SelectedItem);
+
+					// Delete the previously exported .json file for the row that is being deleted
+					if (UNodalOffsetTool* Tool = NodalOffsetTool.Get())
+					{
+						Tool->DeleteExportedRow(SelectedItem->Index);
+					}
 				}
 
 				CalibrationListView->RequestListRefresh();
@@ -2332,7 +2353,7 @@ AActor* UCameraNodalOffsetAlgoPoints::FindFirstCalibrator() const
 	return nullptr;
 }
 
-bool UCameraNodalOffsetAlgoPoints::CalibratorPointCacheFromName(const FString& Name, FCalibratorPointCache& CalibratorPointCache) const
+bool UCameraNodalOffsetAlgoPoints::CalibratorPointCacheFromName(const FString& Name, FNodalOffsetPointsCalibratorPointData& CalibratorPointCache) const
 {
 	CalibratorPointCache.bIsValid = false;
 
@@ -2431,6 +2452,12 @@ void UCameraNodalOffsetAlgoPoints::ClearCalibrationRows()
 	{
 		CalibrationListView->RequestListRefresh();
 	}
+
+	// End the current calibration session (a new one will begin the next time a new row is added)
+	if (UNodalOffsetTool* Tool = NodalOffsetTool.Get())
+	{
+		Tool->EndCalibrationSession();
+	}
 }
 
 void UCameraNodalOffsetAlgoPoints::UndistortCalibrationRowPoints()
@@ -2441,7 +2468,7 @@ void UCameraNodalOffsetAlgoPoints::UndistortCalibrationRowPoints()
 	}
 
 	// If the data for these rows was already undistorted, early-out
-	const TSharedPtr<FCalibrationRowData>& FirstRow = CalibrationRows[0];
+	const TSharedPtr<FNodalOffsetPointsRowData>& FirstRow = CalibrationRows[0];
 	if (FirstRow->bUndistortedIsValid)
 	{
 		return;
@@ -2470,7 +2497,7 @@ void UCameraNodalOffsetAlgoPoints::UndistortCalibrationRowPoints()
 	TArray<FVector2D> ImagePoints;
 	ImagePoints.Reserve(NumPoints);
 
-	for (const TSharedPtr<FCalibrationRowData>& Row : CalibrationRows)
+	for (const TSharedPtr<FNodalOffsetPointsRowData>& Row : CalibrationRows)
 	{
 		ImagePoints.Add(FVector2D(Row->Point2D.X, Row->Point2D.Y));
 	}
@@ -2481,7 +2508,7 @@ void UCameraNodalOffsetAlgoPoints::UndistortCalibrationRowPoints()
 
 	for (int32 Index = 0; Index < NumPoints; ++Index)
 	{
-		const TSharedPtr<FCalibrationRowData>& Row = CalibrationRows[Index];
+		const TSharedPtr<FNodalOffsetPointsRowData>& Row = CalibrationRows[Index];
 		Row->UndistortedPoint2D = UndistortedPoints[Index];
 		Row->bUndistortedIsValid = true;
 	}
@@ -2525,15 +2552,15 @@ bool UCameraNodalOffsetAlgoPoints::GetStepsControllerAndLensFile(
 }
 
 void UCameraNodalOffsetAlgoPoints::GroupRowsByCameraPose(
-	TArray<TSharedPtr<TArray<TSharedPtr<FCalibrationRowData>>>>& OutSamePoseRowGroups,
-	const TArray<TSharedPtr<FCalibrationRowData>>& Rows) const
+	TArray<TSharedPtr<TArray<TSharedPtr<FNodalOffsetPointsRowData>>>>& OutSamePoseRowGroups,
+	const TArray<TSharedPtr<FNodalOffsetPointsRowData>>& Rows) const
 {
-	for (const TSharedPtr<FCalibrationRowData>& Row : Rows)
+	for (const TSharedPtr<FNodalOffsetPointsRowData>& Row : Rows)
 	{
 		check(Row.IsValid());
 
 		// Find the group it belongs to based on transform
-		TSharedPtr<TArray<TSharedPtr<FCalibrationRowData>>> ClosestGroup;
+		TSharedPtr<TArray<TSharedPtr<FNodalOffsetPointsRowData>>> ClosestGroup;
 
 		for (const auto& SamePoseRowGroup : OutSamePoseRowGroups)
 		{
@@ -2546,7 +2573,7 @@ void UCameraNodalOffsetAlgoPoints::GroupRowsByCameraPose(
 
 		if (!ClosestGroup.IsValid())
 		{
-			ClosestGroup = MakeShared<TArray<TSharedPtr<FCalibrationRowData>>>();
+			ClosestGroup = MakeShared<TArray<TSharedPtr<FNodalOffsetPointsRowData>>>();
 			OutSamePoseRowGroups.Add(ClosestGroup);
 		}
 
@@ -2554,6 +2581,121 @@ void UCameraNodalOffsetAlgoPoints::GroupRowsByCameraPose(
 	}
 }
 
+void UCameraNodalOffsetAlgoPoints::ExportSessionData()
+{
+	using namespace UE::CameraCalibration::Private;
+
+	if (UNodalOffsetTool* Tool = NodalOffsetTool.Get())
+	{
+		// Add all data to a json object that is needed to run this algorithm that is NOT part of a specific row
+		TSharedPtr<FJsonObject> JsonSessionData = MakeShared<FJsonObject>();
+
+		JsonSessionData->SetNumberField(NodalOffsetPointsExportFields::Version, DATASET_VERSION);
+
+		// Export the session data to a .json file
+		Tool->ExportSessionData(JsonSessionData.ToSharedRef());
+	}
+}
+
+void UCameraNodalOffsetAlgoPoints::ExportRow(TSharedPtr<FNodalOffsetPointsRowData> Row)
+{
+	if (UNodalOffsetTool* Tool = NodalOffsetTool.Get())
+	{
+		if (const TSharedPtr<FJsonObject>& RowObject = FJsonObjectConverter::UStructToJsonObject<FNodalOffsetPointsRowData>(Row.ToSharedRef().Get()))
+		{
+			// Export the row data to a .json file
+			Tool->ExportCalibrationRow(Row->Index, RowObject.ToSharedRef());
+		}
+	}
+}
+
+bool UCameraNodalOffsetAlgoPoints::HasCalibrationData() const
+{
+	return (CalibrationRows.Num() > 0);
+}
+
+void UCameraNodalOffsetAlgoPoints::PreImportCalibrationData()
+{
+	// Clear the current set of rows before importing new ones
+	CalibrationRows.Empty();
+}
+
+int32 UCameraNodalOffsetAlgoPoints::ImportCalibrationRow(const TSharedRef<FJsonObject>& CalibrationRowObject, const FImage& RowImage)
+{
+	UNodalOffsetTool* Tool = NodalOffsetTool.Get();
+	if (Tool == nullptr)
+	{
+		return -1;
+	}
+
+	const FCameraCalibrationStepsController* StepsController = Tool->GetCameraCalibrationStepsController();
+	if (StepsController == nullptr)
+	{
+		return -1;
+	}
+
+	// Create a new row to populate with data from the Json object
+	TSharedPtr<FNodalOffsetPointsRowData> NewRow = MakeShared<FNodalOffsetPointsRowData>();
+
+	// We enforce strict mode to ensure that every field in the UStruct of row data is present in the imported json.
+	// If any fields are missing, it is likely the row will be invalid, which will lead to errors in the calibration.
+	constexpr int64 CheckFlags = 0;
+	constexpr int64 SkipFlags = 0;
+	constexpr bool bStrictMode = true;
+	if (FJsonObjectConverter::JsonObjectToUStruct<FNodalOffsetPointsRowData>(CalibrationRowObject, NewRow.Get(), CheckFlags, SkipFlags, bStrictMode))
+	{
+		CalibrationRows.Add(NewRow);
+
+		// Set the camera guid to match the currently selected camera
+		if (const ACameraActor* Camera = StepsController->GetCamera())
+		{
+			NewRow->CameraData.UniqueId = Camera->GetUniqueID();
+		}
+
+		// Set the camera parent guid to match the parent of the currently selected camera
+		if (const ACameraActor* Camera = StepsController->GetCamera())
+		{
+			if (const AActor* CameraParentActor = Camera->GetAttachParentActor())
+			{
+				NewRow->CameraData.ParentUniqueId = CameraParentActor->GetUniqueID();
+			}
+		}
+
+		// Set the calibrator guid to match the currently selected calibrator
+		if (AActor* CalibratorActor = Calibrator.Get())
+		{
+			NewRow->CameraData.CalibratorUniqueId = Calibrator->GetUniqueID();
+		}
+
+		// Set the calibrator parent guid to match the parent of the currently selected calibrator
+		if (AActor* CalibratorActor = Calibrator.Get())
+		{
+			if (const AActor* CalibratorParentActor = Calibrator->GetAttachParentActor())
+			{
+				NewRow->CameraData.CalibratorParentUniqueId = CalibratorParentActor->GetUniqueID();
+			}
+		}
+	}
+	else
+	{
+		UE_LOG(LogCameraCalibrationEditor, Warning, TEXT("NodalOffsetPoints algo failed to import calibration row because at least one field could not be deserialized from the json file."));
+	}
+
+	return NewRow->Index;
+}
+
+void UCameraNodalOffsetAlgoPoints::PostImportCalibrationData()
+{
+	// Sort imported calibration rows by row index
+	CalibrationRows.Sort([](const TSharedPtr<FNodalOffsetPointsRowData>& LHS, const TSharedPtr<FNodalOffsetPointsRowData>& RHS) { return LHS->Index < RHS->Index; });
+
+	// Notify the ListView of the new data
+	if (CalibrationListView.IsValid())
+	{
+		CalibrationListView->RequestListRefresh();
+	}
+}
+
 TSharedRef<SWidget> UCameraNodalOffsetAlgoPoints::BuildHelpWidget()
 {
 	return SNew(STextBlock)
diff --git a/Engine/Plugins/VirtualProduction/CameraCalibration/Source/CameraCalibrationEditor/Private/Calibrators/CameraNodalOffsetAlgoPoints.h b/Engine/Plugins/VirtualProduction/CameraCalibration/Source/CameraCalibrationEditor/Private/Calibrators/CameraNodalOffsetAlgoPoints.h
--- a/Engine/Plugins/VirtualProduction/CameraCalibration/Source/CameraCalibrationEditor/Private/Calibrators/CameraNodalOffsetAlgoPoints.h
+++ b/Engine/Plugins/VirtualProduction/CameraCalibration/Source/CameraCalibrationEditor/Private/Calibrators/CameraNodalOffsetAlgoPoints.h
@@ -13,6 +13,7 @@ class FCameraCalibrationStepsController;
 template <typename ItemType>
 class SListView;
 
+class FJsonObject;
 class UCalibrationPointComponent;
 
 template<typename OptionType>
@@ -24,6 +25,119 @@ namespace CameraNodalOffsetAlgoPoints
 	class SCalibrationRowGenerator;
 };
 
+/** Holds information of the identified calibrator 2d point for a given sample of a 2d-3d correlation */
+USTRUCT()
+struct FNodalOffsetPointsCalibratorPointData
+{
+	GENERATED_BODY()
+
+	// True if the rest of the contents are valid.
+	UPROPERTY()
+	bool bIsValid = false;
+
+	// The name of the 3d calibrator point, as defined in the mesh
+	UPROPERTY()
+	FString Name;
+
+	// The world space 3d location of the point
+	UPROPERTY()
+	FVector Location = FVector::ZeroVector;
+};
+
+/** Holds information of the camera pose for a given sample of a 2d-3d correlation */
+USTRUCT()
+struct FNodalOffsetPointsCameraData
+{
+	GENERATED_BODY()
+
+	// True if the rest of the contents are valid.
+	UPROPERTY()
+	bool bIsValid = false;
+
+	// The unique id of the camera object. Used to detect camera selection changes during a calibration session.
+	UPROPERTY()
+	uint32 UniqueId = INDEX_NONE;
+
+	// The camera pose
+	UPROPERTY()
+	FTransform Pose = FTransform::Identity;
+
+	// True if the cached camera pose included a nodal offset
+	UPROPERTY()
+	bool bWasNodalOffsetApplied = false;
+
+	// True if distortion was evaluated (using this LensFile) 
+	UPROPERTY()
+	bool bWasDistortionEvaluated = false;
+
+	// Input focus from lens file evaluation data
+	UPROPERTY()
+	float InputFocus = 0.0f;
+
+	// Input zoom from lens file evaluation data
+	UPROPERTY()
+	float InputZoom = 0.0f;
+
+	// The parent pose (expected to be the tracker origin)
+	UPROPERTY()
+	FTransform ParentPose = FTransform::Identity;
+
+	// The parent unique id
+	UPROPERTY()
+	uint32 ParentUniqueId = INDEX_NONE;
+
+	// Calibrator Pose
+	UPROPERTY()
+	FTransform CalibratorPose = FTransform::Identity;
+
+	// Calibrator ParentPose
+	UPROPERTY()
+	FTransform CalibratorParentPose = FTransform::Identity;
+
+	// Calibrator ComponentPose
+	UPROPERTY()
+	TMap<uint32, FTransform> CalibratorComponentPoses;
+
+	// Calibrator unique id
+	UPROPERTY()
+	uint32 CalibratorUniqueId = INDEX_NONE;
+
+	// Calibrator parent unique id
+	UPROPERTY()
+	uint32 CalibratorParentUniqueId = INDEX_NONE;
+};
+
+/** Holds information of the calibrator 3d point for a given sample of a 2d-3d correlation */
+USTRUCT()
+struct FNodalOffsetPointsRowData
+{
+	GENERATED_BODY()
+
+	// Index of this calibration row
+	UPROPERTY()
+	int32 Index = INDEX_NONE;
+
+	// Normalized 0~1 2d location of the identified calibrator point in the media plate.
+	UPROPERTY()
+	FVector2D Point2D = FVector2D::ZeroVector;
+
+	// Location of Point2D after it has been undistorted
+	UPROPERTY()
+	FVector2D UndistortedPoint2D = FVector2D::ZeroVector;
+
+	// True if this row has had its Point2D undistorted
+	UPROPERTY()
+	bool bUndistortedIsValid = false;
+
+	// Holds information of the calibrator point data for this sample.
+	UPROPERTY()
+	FNodalOffsetPointsCalibratorPointData CalibratorPointData;
+
+	// Holds information of the camera data for this sample
+	UPROPERTY()
+	FNodalOffsetPointsCameraData CameraData;
+};
+
 /** 
  * Implements a nodal offset calibration algorithm. It uses 3d points (UCalibrationPointComponent) 
  * specified in a selectable calibrator with features that the user can identify by clicking on the 
@@ -49,7 +163,12 @@ public:
 	virtual TSharedRef<SWidget> BuildUI() override;
 	virtual bool GetNodalOffset(FNodalPointOffset& OutNodalOffset, float& OutFocus, float& OutZoom, float& OutError, FText& OutErrorMessage) override;
 	virtual FName FriendlyName() const override { return TEXT("Nodal Offset Points Method"); };
+	virtual FName ShortName() const override { return TEXT("Points"); };
 	virtual void OnSavedNodalOffset() override;
+	virtual bool HasCalibrationData() const override;
+	virtual void PreImportCalibrationData() override;
+	virtual int32 ImportCalibrationRow(const TSharedRef<FJsonObject>& CalibrationRowObject, const FImage& RowImage) override;
+	virtual void PostImportCalibrationData() override;
 	virtual TSharedRef<SWidget> BuildHelpWidget() override;
 	//~ End CalibPointsNodalOffsetAlgo
 
@@ -69,82 +188,9 @@ protected:
 		FString Name;
 	};
 
-	/** Holds information of the identified calibrator 2d point for a given sample of a 2d-3d correlation */
-	struct FCalibratorPointCache
-	{
-		// True if the rest of the contents are valid.
-		bool bIsValid;
-
-		// The name of the 3d calibrator point, as defined in the mesh
-		FString Name;
-
-		// The world space 3d location of the point
-		FVector Location;
-	};
-
-	/** Holds information of the camera pose for a given sample of a 2d-3d correlation */
-	struct FCameraDataCache
-	{
-		// True if the rest of the contents are valid.
-		bool bIsValid;
-
-		// The unique id of the camera object. Used to detect camera selection changes during a calibration session.
-		uint32 UniqueId;
-
-		// The camera pose
-		FTransform Pose;
-
-		// True if the cached camera pose included a nodal offset
-		bool bWasNodalOffsetApplied;
-
-		// True if distortion was evaluated (using this LensFile) 
-		bool bWasDistortionEvaluated;
-
-		// The data used to evaluate the lens data in the camera for this sample
-		FLensFileEvalData LensFileEvalData;
-
-		// The parent pose (expected to be the tracker origin)
-		FTransform ParentPose;
-
-		// The parent unique id
-		uint32 ParentUniqueId;
-
-		// Calibrator Pose
-		FTransform CalibratorPose;
-
-		// Calibrator ParentPose
-		FTransform CalibratorParentPose;
-
-		// Calibrator ComponentPose
-		TMap<uint32, FTransform> CalibratorComponentPoses;
-
-		// Calibrator unique id
-		uint32 CalibratorUniqueId;
-
-		// Calibrator parent unique id
-		uint32 CalibratorParentUniqueId;
-	};
-
-	/** Holds information of the calibrator 3d point for a given sample of a 2d-3d correlation */
-	struct FCalibrationRowData
-	{
-		// Normalized 0~1 2d location of the identified calibrator point in the media plate.
-		FVector2D Point2D;
-
-		// Location of Point2D after it has been undistorted
-		FVector2D UndistortedPoint2D;
-
-		// True if this row has had its Point2D undistorted
-		bool bUndistortedIsValid = false;
-
-		// Holds information of the calibrator point data for this sample.
-		FCalibratorPointCache CalibratorPointData;
-
-		// Holds information of the camera data for this sample
-		FCameraDataCache CameraData;
-	};
-
 protected:
+	/** Version of the current calibration dataset */
+	static const int DATASET_VERSION;
 
 	/** The nodal offset tool controller */
 	TWeakObjectPtr<UNodalOffsetTool> NodalOffsetTool;
@@ -162,16 +208,16 @@ protected:
 	TSharedPtr<SComboBox<TSharedPtr<FCalibratorPointData>>> CalibratorPointsComboBox;
 
 	/** Rows source for the CalibrationListView */
-	TArray<TSharedPtr<FCalibrationRowData>> CalibrationRows;
+	TArray<TSharedPtr<FNodalOffsetPointsRowData>> CalibrationRows;
 
 	/** Displays the list of calibration points that will be used to calculate the nodal offset */
-	TSharedPtr<SListView<TSharedPtr<FCalibrationRowData>>> CalibrationListView;
+	TSharedPtr<SListView<TSharedPtr<FNodalOffsetPointsRowData>>> CalibrationListView;
 
 	/** Caches the last calibrator point 3d location.  Will hold last value before the nodal offset tool is paused */
-	TArray<FCalibratorPointCache> LastCalibratorPoints;
+	TArray<FNodalOffsetPointsCalibratorPointData> LastCalibratorPoints;
 
 	/** Caches the last camera data.  Will hold last value before the nodal offset tool is paused */
-	FCameraDataCache LastCameraData;
+	FNodalOffsetPointsCameraData LastCameraData;
 
 	/** Instructs the inline allocation policy how many elements to allocate in a single allocation */
 	static constexpr uint32 NumInlineAllocations = 32;
@@ -221,7 +267,7 @@ protected:
 	 * 
 	 * @return True if successful.
 	 */
-	bool CalibratorPointCacheFromName(const FString& Name, FCalibratorPointCache& CalibratorPointCache) const;
+	bool CalibratorPointCacheFromName(const FString& Name, FNodalOffsetPointsCalibratorPointData& CalibratorPointCache) const;
 
 	/** Returns the world 3d location of the currently selected calibrator */
 	bool GetCurrentCalibratorPointLocation(FVector& OutLocation);
@@ -239,7 +285,7 @@ protected:
 	void OnObjectsReplaced(const TMap<UObject*, UObject*>& OldToNewInstanceMap);
 
 	/** Validates a new calibration point to determine if it should be added as a new sample row */
-	virtual bool ValidateNewRow(TSharedPtr<FCalibrationRowData>& Row, FText& OutErrorMessage) const;
+	virtual bool ValidateNewRow(TSharedPtr<FNodalOffsetPointsRowData>& Row, FText& OutErrorMessage) const;
 
 	/** Applies the nodal offset to the calibrator */
 	bool ApplyNodalOffsetToCalibrator();
@@ -254,7 +300,7 @@ protected:
 	bool ApplyNodalOffsetToCalibratorComponents();
 
 	/** Does basic checks on the data before performing the actual calibration */
-	bool BasicCalibrationChecksPass(const TArray<TSharedPtr<FCalibrationRowData>>& Rows, FText& OutErrorMessage) const;
+	bool BasicCalibrationChecksPass(const TArray<TSharedPtr<FNodalOffsetPointsRowData>>& Rows, FText& OutErrorMessage) const;
 
 	/** 
 	 * Gets the step controller and the lens file. 
@@ -277,7 +323,7 @@ protected:
 	 */
 	bool CalculatedOptimalCameraComponentPose(
 		FTransform& OutDesiredCameraTransform, 
-		const TArray<TSharedPtr<FCalibrationRowData>>& Rows, 
+		const TArray<TSharedPtr<FNodalOffsetPointsRowData>>& Rows, 
 		FText& OutErrorMessage) const;
 
 	/**
@@ -297,7 +343,7 @@ protected:
 		float& OutFocus, 
 		float& OutZoom, 
 		float& OutError, 
-		const TArray<TSharedPtr<FCalibrationRowData>>& Rows, 
+		const TArray<TSharedPtr<FNodalOffsetPointsRowData>>& Rows, 
 		FText& OutErrorMessage) const;
 
 	/**
@@ -307,8 +353,8 @@ protected:
 	 * @param Rows Ungrouped calibration samples.
 	 */
 	void GroupRowsByCameraPose(
-		TArray<TSharedPtr<TArray<TSharedPtr<FCalibrationRowData>>>>& OutSamePoseRowGroups, 
-		const TArray<TSharedPtr<FCalibrationRowData>>& Rows) const;
+		TArray<TSharedPtr<TArray<TSharedPtr<FNodalOffsetPointsRowData>>>>& OutSamePoseRowGroups, 
+		const TArray<TSharedPtr<FNodalOffsetPointsRowData>>& Rows) const;
 
 	/** 
 	 * Detects if the calibrator moved significantly in the given sample points 
@@ -317,7 +363,7 @@ protected:
 	 * 
 	 * @return True if the calibrator moved significantly across all the samples.
 	 */
-	bool CalibratorMovedAcrossGroups(const TArray<TSharedPtr<TArray<TSharedPtr<FCalibrationRowData>>>>& OutSamePoseRowGroups) const;
+	bool CalibratorMovedAcrossGroups(const TArray<TSharedPtr<TArray<TSharedPtr<FNodalOffsetPointsRowData>>>>& OutSamePoseRowGroups) const;
 
 	/**
 	 * Detects if the calibrator moved significantly in the given sample points
@@ -326,7 +372,7 @@ protected:
 	 *
 	 * @return True if the calibrator moved significantly across all the samples.
 	 */
-	bool CalibratorMovedInAnyRow(const TArray<TSharedPtr<FCalibrationRowData>>& Rows) const;
+	bool CalibratorMovedInAnyRow(const TArray<TSharedPtr<FNodalOffsetPointsRowData>>& Rows) const;
 
 	/**
 	 * Calculates the optimal camera parent transform that minimizes the reprojection of the sample points.
@@ -338,7 +384,7 @@ protected:
 	 * @return True if successful
 	 */
 	bool CalcTrackingOriginPoseForSingleCamPose(
-		const TArray<TSharedPtr<FCalibrationRowData>>& Rows,
+		const TArray<TSharedPtr<FNodalOffsetPointsRowData>>& Rows,
 		FTransform& OutTransform,
 		FText& OutErrorMessage);
 
@@ -352,13 +398,19 @@ protected:
 	 * @return True if successful
 	 */
 	bool CalcCalibratorPoseForSingleCamPose(
-		const TArray<TSharedPtr<FCalibrationRowData>>& Rows,
+		const TArray<TSharedPtr<FNodalOffsetPointsRowData>>& Rows,
 		FTransform& OutTransform,
 		FText& OutErrorMessage);
 
 	/* Refine the input nodal offset by running a minimization algorithm designed to minimize the sum of the reprojection errors for the camera views in the set of input pose groups*/
-	double MinimizeReprojectionError(FTransform& InOutNodalOffset, const TArray<TSharedPtr<TArray<TSharedPtr<FCalibrationRowData>>>>& SamePoseRowGroups) const;
+	double MinimizeReprojectionError(FTransform& InOutNodalOffset, const TArray<TSharedPtr<TArray<TSharedPtr<FNodalOffsetPointsRowData>>>>& SamePoseRowGroups) const;
 
 	/* Compute the reprojection error for the rows in the input PoseGroup using the input NodalOffset */
-	double ComputeReprojectionError(const FTransform& NodalOffset, const TArray<TSharedPtr<FCalibrationRowData>>& PoseGroup) const;
+	double ComputeReprojectionError(const FTransform& NodalOffset, const TArray<TSharedPtr<FNodalOffsetPointsRowData>>& PoseGroup) const;
+
+	/** Export global session data to a .json file */
+	void ExportSessionData();
+
+	/** Export the row data to a json file */
+	void ExportRow(TSharedPtr<FNodalOffsetPointsRowData> Row);
 };
diff --git a/Engine/Plugins/VirtualProduction/CameraCalibrationCore/Source/CameraCalibrationCore/Public/CameraCalibrationSettings.h b/Engine/Plugins/VirtualProduction/CameraCalibrationCore/Source/CameraCalibrationCore/Public/CameraCalibrationSettings.h
--- a/Engine/Plugins/VirtualProduction/CameraCalibrationCore/Source/CameraCalibrationCore/Public/CameraCalibrationSettings.h
+++ b/Engine/Plugins/VirtualProduction/CameraCalibrationCore/Source/CameraCalibrationCore/Public/CameraCalibrationSettings.h
@@ -73,6 +73,9 @@ public:
 	/** Get the default MaterialInterface used by the input Model Handler class to apply the post-process lens distortion effect */
 	UMaterialInterface* GetDefaultDistortionMaterial(const TSubclassOf<ULensDistortionModelHandlerBase>& InModelHandler) const;
 
+	/** Returns true if the calibration dataset import and export featuers are enabled, false otherwise */
+	bool IsCalibrationDatasetImportExportEnabled() const { return bEnableCalibrationDatasetImportExport; }
+
 #if WITH_EDITOR
 protected:
 	FOnDisplacementMapResolutionChanged DisplacementMapResolutionChangedDelegate;
@@ -116,6 +119,10 @@ private:
 	TMap<FName, TSoftObjectPtr<UMaterialInterface>> CalibrationOverlayMaterialOverrides;
 #endif
 
+	/** Setting to toggle the calibration dataset import and export features */
+	UPROPERTY(config, EditAnywhere, Category = "Settings")
+	bool bEnableCalibrationDatasetImportExport = true;
+
 private:
 	/** Delegate handle to run after the engine is initialized */
 	FDelegateHandle PostEngineInitHandle;
diff --git a/Engine/Plugins/VirtualProduction/CameraCalibrationCore/Source/CameraCalibrationCore/Public/CameraLensDistortionAlgo.h b/Engine/Plugins/VirtualProduction/CameraCalibrationCore/Source/CameraCalibrationCore/Public/CameraLensDistortionAlgo.h
--- a/Engine/Plugins/VirtualProduction/CameraCalibrationCore/Source/CameraCalibrationCore/Public/CameraLensDistortionAlgo.h
+++ b/Engine/Plugins/VirtualProduction/CameraCalibrationCore/Source/CameraCalibrationCore/Public/CameraLensDistortionAlgo.h
@@ -13,8 +13,10 @@ struct FPointerEvent;
 
 struct FDistortionInfo;
 struct FFocalLengthInfo;
+struct FImage;
 struct FImageCenterInfo;
 
+class FJsonObject;
 class ULensDistortionTool;
 class ULensModel;
 class UMaterialInterface;
@@ -46,9 +48,12 @@ public:
 	/** Returns the UI of this calibrator. Expected to only be called once */
 	virtual TSharedRef<SWidget> BuildUI() { return SNullWidget::NullWidget; };
 
-	/** Returns a descriptive name/title of this nodal offset algorithm */
+	/** Returns a descriptive name/title for this algorithm */
 	virtual FName FriendlyName() const { return TEXT("Invalid Name"); };
 
+	/** Returns a shorter name for this algorithm */
+	virtual FName ShortName() const { return TEXT("Invalid"); };
+
 	/** Returns the overlay material used by this algo (if any) */
 	virtual UMaterialInterface* GetOverlayMaterial() const { return nullptr; };
 
@@ -72,6 +77,25 @@ public:
 		return false; 
 	};
 
+	/** Returns true if there is any existing calibration data */
+	virtual bool HasCalibrationData() const { return false; };
+
+	/** Performs any necessary steps (such as clearing existing calibration data) before importing a calibration dataset */
+	virtual void PreImportCalibrationData() { };
+
+	/** Import a JsonObject of calibration data that is needed by the algorithm, but is not associated with a single row of data */
+	virtual void ImportSessionData(const TSharedRef<FJsonObject>& SessionDataObject) { };
+
+	/** 
+	 * Import a JsonObject of calibration data that represents a single calibration row. 
+	 * RowImage is not guaranteed to have any data. The algorithm should check its validity before using it.
+	 * Returns the row index of the imported row.
+	 */
+	virtual int32 ImportCalibrationRow(const TSharedRef<FJsonObject>& CalibrationRowObject, const FImage& RowImage) { return -1; };
+
+	/** Performs any necessary steps after importing a calibration dataset */
+	virtual void PostImportCalibrationData() { };
+
 	/** Called to present the user with instructions on how to use this algo */
 	virtual TSharedRef<SWidget> BuildHelpWidget() { return SNew(STextBlock).Text(FText::FromString(TEXT("Coming soon!"))); };
 };
diff --git a/Engine/Plugins/VirtualProduction/CameraCalibrationCore/Source/CameraCalibrationCore/Public/CameraNodalOffsetAlgo.h b/Engine/Plugins/VirtualProduction/CameraCalibrationCore/Source/CameraCalibrationCore/Public/CameraNodalOffsetAlgo.h
--- a/Engine/Plugins/VirtualProduction/CameraCalibrationCore/Source/CameraCalibrationCore/Public/CameraNodalOffsetAlgo.h
+++ b/Engine/Plugins/VirtualProduction/CameraCalibrationCore/Source/CameraCalibrationCore/Public/CameraNodalOffsetAlgo.h
@@ -10,10 +10,12 @@
 #include "CameraNodalOffsetAlgo.generated.h"
 
 struct FGeometry;
+struct FImage;
 struct FKey;
 struct FNodalPointOffset;
 struct FPointerEvent;
 
+class FJsonObject;
 class UMaterialInterface;
 class UNodalOffsetTool;
 class SWidget;
@@ -50,9 +52,12 @@ public:
 	/** Returns the most recently calibrated nodal offset transform, with an error metric */
 	virtual bool GetNodalOffset(FNodalPointOffset& OutNodalOffset, float& OutFocus, float& OutZoom, float& OutError, FText& OutErrorMessage) { return false; };
 
-	/** Returns a descriptive name/title of this nodal offset algorithm */
+	/** Returns a descriptive name/title for this algorithm */
 	virtual FName FriendlyName() const { return TEXT("Invalid Name"); };
 
+	/** Returns a shorter name for this algorithm */
+	virtual FName ShortName() const { return TEXT("Invalid"); };
+
 	/** Returns the overlay material used by this algo (if any) */
 	virtual UMaterialInterface* GetOverlayMaterial() const { return nullptr; };
 
@@ -62,6 +67,25 @@ public:
 	/** Called when the current offset was saved */
 	virtual void OnSavedNodalOffset() { };
 
+	/** Returns true if there is any existing calibration data */
+	virtual bool HasCalibrationData() const { return false; };
+
+	/** Performs any necessary steps (such as clearing existing calibration data) before importing a calibration dataset */
+	virtual void PreImportCalibrationData() { };
+
+	/** Import a JsonObject of calibration data that is needed by the algorithm, but is not associated with a single row of data */
+	virtual void ImportSessionData(const TSharedRef<FJsonObject>& SessionDataObject) { };
+
+	/**
+	 * Import a JsonObject of calibration data that represents a single calibration row.
+	 * RowImage is not guaranteed to have any data. The algorithm should check its validity before using it.
+	 * Returns the row index of the imported row.
+	 */
+	virtual int32 ImportCalibrationRow(const TSharedRef<FJsonObject>& CalibrationRowObject, const FImage& RowImage) { return -1; };
+
+	/** Performs any necessary steps after importing a calibration dataset */
+	virtual void PostImportCalibrationData() { };
+
 	/** Called to present the user with instructions on how to use this algo */
 	virtual TSharedRef<SWidget> BuildHelpWidget() { return SNew(STextBlock).Text(FText::FromString(TEXT("Coming soon!")));  };
 };
