[CameraCalibration] Fix calibration algos not functioning if a LensFile is not specified in the LiveLink camera controller

diff --git a/Engine/Plugins/VirtualProduction/CameraCalibration/CameraCalibration.uplugin b/Engine/Plugins/VirtualProduction/CameraCalibration/CameraCalibration.uplugin
--- a/Engine/Plugins/VirtualProduction/CameraCalibration/CameraCalibration.uplugin
+++ b/Engine/Plugins/VirtualProduction/CameraCalibration/CameraCalibration.uplugin
@@ -25,10 +25,6 @@
 			"Name": "CameraCalibrationCore",
 			"Enabled": true
 		},
-		{
-			"Name": "LiveLink",
-			"Enabled": true
-		},
 		{
 			"Name": "Composure",
 			"Enabled": true
@@ -40,10 +36,6 @@
 		{
 			"Name": "OpenCVLensDistortion",
 			"Enabled": true
-		},
-		{
-			"Name": "LiveLinkCamera",
-			"Enabled": true
 		}
 	]
 }
\ No newline at end of file
diff --git a/Engine/Plugins/VirtualProduction/CameraCalibration/Source/CameraCalibrationEditor/CameraCalibrationEditor.Build.cs b/Engine/Plugins/VirtualProduction/CameraCalibration/Source/CameraCalibrationEditor/CameraCalibrationEditor.Build.cs
--- a/Engine/Plugins/VirtualProduction/CameraCalibration/Source/CameraCalibrationEditor/CameraCalibrationEditor.Build.cs
+++ b/Engine/Plugins/VirtualProduction/CameraCalibration/Source/CameraCalibrationEditor/CameraCalibrationEditor.Build.cs
@@ -28,10 +28,6 @@ namespace UnrealBuildTool.Rules
 					"InputCore",
 					"Json",
 					"JsonUtilities",
-					"LiveLinkCamera",
-					"LiveLinkComponents",
-					"LiveLinkEditor",
-					"LiveLinkInterface",
 					"MediaAssets",
 					"MediaFrameworkUtilities",
 					"OpenCV",
diff --git a/Engine/Plugins/VirtualProduction/CameraCalibration/Source/CameraCalibrationEditor/Private/AssetEditor/CameraCalibrationStepsController.cpp b/Engine/Plugins/VirtualProduction/CameraCalibration/Source/CameraCalibrationEditor/Private/AssetEditor/CameraCalibrationStepsController.cpp
--- a/Engine/Plugins/VirtualProduction/CameraCalibration/Source/CameraCalibrationEditor/Private/AssetEditor/CameraCalibrationStepsController.cpp
+++ b/Engine/Plugins/VirtualProduction/CameraCalibration/Source/CameraCalibrationEditor/Private/AssetEditor/CameraCalibrationStepsController.cpp
@@ -30,8 +30,6 @@
 #include "Kismet/KismetRenderingLibrary.h"
 #include "LensComponent.h"
 #include "LensFile.h"
-#include "LiveLinkCameraController.h"
-#include "LiveLinkComponentController.h"
 #include "Materials/Material.h"
 #include "Materials/MaterialInstanceDynamic.h"
 #include "MediaPlayer.h"
@@ -72,36 +70,6 @@ namespace CameraCalibrationStepsController
 		}
 		return nullptr;
 	}
-
-	/** Retrieves the latest lens evaluation data from the LiveLink controller in the given camera for the given lens file */
-	const FLensFileEvalData* LensFileEvalDataFromCamera(const ACameraActor* InCamera, const ULensFile* InLensFile)
-	{
-		if (!InCamera || !InLensFile)
-		{
-			return nullptr;
-		}
-
-		TArray<ULiveLinkComponentController*> LLComponentControllers;
-		InCamera->GetComponents(LLComponentControllers);
-
-		for (const ULiveLinkComponentController* LLComponentController : LLComponentControllers)
-		{
-			for (const TPair<TSubclassOf<ULiveLinkRole>, TObjectPtr<ULiveLinkControllerBase>> Pair : LLComponentController->ControllerMap)
-			{
-				if (ULiveLinkCameraController* CameraController = Cast<ULiveLinkCameraController>(Pair.Value))
-				{
-					const FLensFileEvalData* OutLensFileEvalData = &CameraController->GetLensFileEvalDataRef();
-
-					if (OutLensFileEvalData->LensFile == InLensFile)
-					{
-						return OutLensFileEvalData;
-					}
-				}
-			}
-		}
-
-		return nullptr;
-	}
 }
 
 
@@ -207,7 +175,11 @@ TSharedPtr<SWidget> FCameraCalibrationStepsController::BuildUI()
 bool FCameraCalibrationStepsController::OnTick(float DeltaTime)
 {
 	// Update the lens file eval data
-	LensFileEvalData = CameraCalibrationStepsController::LensFileEvalDataFromCamera(Camera.Get(), LensFile.Get());
+	LensFileEvaluationInputs.bIsValid = false;
+	if (const ULensComponent* const LensComponent = FindLensComponent())
+	{
+		LensFileEvaluationInputs = LensComponent->GetLensFileEvaluationInputs();
+	}
 
 	for (TStrongObjectPtr<UCameraCalibrationStep>& Step : CalibrationSteps)
 	{
@@ -1053,66 +1025,30 @@ bool FCameraCalibrationStepsController::AreMediaPlaybackControlsVisible() const
 	return bShowMediaPlaybackButtons;
 }
 
-ULiveLinkCameraController* FCameraCalibrationStepsController::FindLiveLinkCameraController() const
-{
-	return FindLiveLinkCameraControllerWithLens(Camera.Get(), LensFile.Get());
-}
-
-ULiveLinkCameraController* FCameraCalibrationStepsController::FindLiveLinkCameraControllerWithLens(const ACameraActor* InCamera, const ULensFile* InLensFile) const
+ACameraActor* FCameraCalibrationStepsController::FindFirstCameraWithCurrentLens() const
 {
-	if (!InCamera || !InLensFile)
+	// We iterate over all cameras in the scene and try to find one that is using the current LensFile
+	ACineCameraActor* FirstCamera = nullptr;
+	for (TActorIterator<ACineCameraActor> CameraItr(GetWorld()); CameraItr; ++CameraItr)
 	{
-		return nullptr;
-	}
+		ACineCameraActor* CameraActor = *CameraItr;
 
-	FAssetData LensFileAssetData(InLensFile);
-
-	TArray<ULiveLinkComponentController*> LiveLinkComponents;
-	InCamera->GetComponents(LiveLinkComponents);
-
-	for (const ULiveLinkComponentController* LiveLinkComponent : LiveLinkComponents)
-	{
-		for (auto It = LiveLinkComponent->ControllerMap.CreateConstIterator(); It; ++It)
+		if (ULensComponent* FoundLensComponent = FindLensComponentOnCamera(CameraActor))
 		{
-			ULiveLinkCameraController* CameraController = Cast<ULiveLinkCameraController>(It->Value);
-
-			if (!CameraController)
-			{
-				continue;
-			}
-
-			const ULensFile* CameraLensFile = CameraController->LensFilePicker.GetLensFile();
-
-			if (!CameraLensFile)
+			if (FirstCamera == nullptr)
 			{
-				continue;
+				FirstCamera = CameraActor;
 			}
-
-			if (FAssetData(CameraLensFile, true).PackageName == LensFileAssetData.PackageName)
+			else
 			{
-				return CameraController;
+				FText ErrorMessage = LOCTEXT("MoreThanOneCameraFoundError", "There are multiple cameras in the scene using this LensFile. When the asset editor opens, be sure to select the correct camera if not already selected.");
+				FMessageDialog::Open(EAppMsgType::Ok, ErrorMessage);
+				break;
 			}
 		}
 	}
 
-	return nullptr;
-}
-
-ACameraActor* FCameraCalibrationStepsController::FindFirstCameraWithCurrentLens() const
-{
-	// We iterate over all cameras in the scene and try to find one that is using the current LensFile
-
-	for (TActorIterator<ACameraActor> CameraItr(GetWorld()); CameraItr; ++CameraItr)
-	{
-		ACameraActor* CameraActor = *CameraItr;
-
-		if (FindLiveLinkCameraControllerWithLens(CameraActor, LensFile.Get()))
-		{
-			return CameraActor;
-		}
-	}
-
-	return nullptr;
+	return FirstCamera;
 }
 
 void FCameraCalibrationStepsController::TogglePlay()
@@ -1164,9 +1100,9 @@ bool FCameraCalibrationStepsController::IsPaused() const
 	return true;
 }
 
-const FLensFileEvalData* FCameraCalibrationStepsController::GetLensFileEvalData() const
+FLensFileEvaluationInputs FCameraCalibrationStepsController::GetLensFileEvaluationInputs() const
 {
-	return LensFileEvalData;
+	return LensFileEvaluationInputs;
 }
 
 ULensFile* FCameraCalibrationStepsController::GetLensFile() const
@@ -1179,17 +1115,14 @@ ULensFile* FCameraCalibrationStepsController::GetLensFile() const
 	return nullptr;
 }
 
-ULensComponent* FCameraCalibrationStepsController::FindLensComponent() const
+ULensComponent* FCameraCalibrationStepsController::FindLensComponentOnCamera(ACameraActor* CineCamera) const
 {
-	ACineCameraActor* CineCamera = Cast<ACineCameraActor>(GetCamera());
 	const ULensFile* OpenLensFile = GetLensFile();
 	if (CineCamera && OpenLensFile)
 	{
 		TInlineComponentArray<ULensComponent*> LensComponents;
 		CineCamera->GetComponents(LensComponents);
 
-		const FText TitleError = LOCTEXT("CameraCalibrationEditorError", "Lens File Error");
-
 		for (ULensComponent* LensComponent : LensComponents)
 		{
 			if (LensComponent->GetLensFile() == OpenLensFile)
@@ -1202,6 +1135,11 @@ ULensComponent* FCameraCalibrationStepsController::FindLensComponent() const
 	return nullptr;
 }
 
+ULensComponent* FCameraCalibrationStepsController::FindLensComponent() const
+{
+	return FindLensComponentOnCamera(GetCamera());
+}
+
 const ULensDistortionModelHandlerBase* FCameraCalibrationStepsController::GetDistortionHandler() const
 {
 	if (ULensComponent* LensComponent = FindLensComponent())
@@ -1280,11 +1218,6 @@ void FCameraCalibrationStepsController::FindMediaSourceUrls(TArray<TSharedPtr<FS
 	}
 }
 
-const FLensFileEvalData* FCameraCalibrationStepsController::GetLensFileEvalData()
-{
-	return LensFileEvalData;
-}
-
 const TConstArrayView<TStrongObjectPtr<UCameraCalibrationStep>> FCameraCalibrationStepsController::GetCalibrationSteps() const
 {
 	return TConstArrayView<TStrongObjectPtr<UCameraCalibrationStep>>(CalibrationSteps);
diff --git a/Engine/Plugins/VirtualProduction/CameraCalibration/Source/CameraCalibrationEditor/Private/AssetEditor/CameraCalibrationStepsController.h b/Engine/Plugins/VirtualProduction/CameraCalibration/Source/CameraCalibrationEditor/Private/AssetEditor/CameraCalibrationStepsController.h
--- a/Engine/Plugins/VirtualProduction/CameraCalibration/Source/CameraCalibrationEditor/Private/AssetEditor/CameraCalibrationStepsController.h
+++ b/Engine/Plugins/VirtualProduction/CameraCalibration/Source/CameraCalibrationEditor/Private/AssetEditor/CameraCalibrationStepsController.h
@@ -4,10 +4,11 @@
 
 #include "CoreMinimal.h"
 
+#include "Containers/Ticker.h"
 #include "Engine/TextureRenderTarget2D.h"
 #include "Engine/World.h"
+#include "LensFile.h"
 #include "UObject/StrongObjectPtr.h"
-#include "Containers/Ticker.h"
 
 
 class ACameraActor;
@@ -21,10 +22,8 @@ class UCompositingElementMaterialPass;
 class ULensComponent;
 class ULensDistortionModelHandlerBase;
 class ULensFile;
-class ULiveLinkCameraController;
 
 struct FGeometry;
-struct FLensFileEvalData;
 struct FPointerEvent;
 
 /** Enumeration of overlay passes used to indicate which overlay pass to interact with */
@@ -90,7 +89,7 @@ public:
 	void Pause();
 
 	/** Returns the latest data used when evaluating the lens */
-	const FLensFileEvalData* GetLensFileEvalData() const;
+	FLensFileEvaluationInputs GetLensFileEvaluationInputs() const;
 
 	/** Returns the LensFile that this tool is using */
 	ULensFile* GetLensFile() const;
@@ -101,9 +100,6 @@ public:
 	/** Returns the distortion handler used to distort the CG being displayed in the simulcam viewport */
 	const ULensDistortionModelHandlerBase* GetDistortionHandler() const;
 
-	/** Finds the LiveLinkCameraController used in the given CameraActor that is also using the given LensFile */
-	ULiveLinkCameraController* FindLiveLinkCameraController() const;
-
 	/** Sets the media source url to be played. Returns true if the url is a valid media source */
 	bool SetMediaSourceUrl(const FString& InMediaSourceUrl);
 
@@ -113,9 +109,6 @@ public:
 	/** Gets the current media source url being played. Empty if None */
 	FString GetMediaSourceUrl() const;
 
-	/** Gets the lens eval data for this frame that was cached in Tick */
-	const FLensFileEvalData* GetLensFileEvalData();
-
 	/** Returns the calibration steps */
 	const TConstArrayView<TStrongObjectPtr<UCameraCalibrationStep>> GetCalibrationSteps() const;
 
@@ -195,8 +188,8 @@ public:
 
 private:
 
-	/** Finds the LiveLinkCameraController used in the given CameraActor that is also using the given LensFile */
-	ULiveLinkCameraController* FindLiveLinkCameraControllerWithLens(const ACameraActor* CameraActor, const ULensFile* InLensFile) const;
+	/** Returns the first lens component with a matching LensFile found on the input camera, or nullptr if none exists */
+	ULensComponent* FindLensComponentOnCamera(ACameraActor* CineCamera) const;
 
 	/** Returns a namespaced version of the given name. Useful to generate names unique to this lens file */
 	FString NamespacedName(const FString&& Name) const;
@@ -296,8 +289,8 @@ private:
 	/** The delegate for the core ticker callback */
 	FTSTicker::FDelegateHandle TickerHandle;
 
-	/** Pointer to the LensFileEvalData used in the current frame. Only valid during the current frame. */
-	const FLensFileEvalData* LensFileEvalData;
+	/** Evaluation Data supplied by the Lens Component for the current frame. Only valid during the current frame. */
+	FLensFileEvaluationInputs LensFileEvaluationInputs;
 
 	/** Setting to control whether the media playback buttons are visible */
 	bool bShowMediaPlaybackButtons = true;
diff --git a/Engine/Plugins/VirtualProduction/CameraCalibration/Source/CameraCalibrationEditor/Private/AssetEditor/CameraCalibrationTimeSliderController.cpp b/Engine/Plugins/VirtualProduction/CameraCalibration/Source/CameraCalibrationEditor/Private/AssetEditor/CameraCalibrationTimeSliderController.cpp
--- a/Engine/Plugins/VirtualProduction/CameraCalibration/Source/CameraCalibrationEditor/Private/AssetEditor/CameraCalibrationTimeSliderController.cpp
+++ b/Engine/Plugins/VirtualProduction/CameraCalibration/Source/CameraCalibrationEditor/Private/AssetEditor/CameraCalibrationTimeSliderController.cpp
@@ -7,7 +7,6 @@
 #include "Styling/AppStyle.h"
 #include "Fonts/FontMeasure.h"
 #include "LensFile.h"
-#include "LiveLinkCameraController.h"
 
 /** Utility struct for converting between scrub range space and local/absolute screen space */
 struct FCameraCalibrationTimeSliderController::FScrubRangeToScreen
@@ -215,18 +214,16 @@ void FCameraCalibrationTimeSliderController::Tick(float DeltaTime)
 		return;
 	}
 
-	ULiveLinkCameraController* CameraController = CalibrationStepsController->FindLiveLinkCameraController();
-	if (!CameraController)
+	const ELensDataCategory SelectedCategoryValue = SelectedCategory.GetValue();
+	const FLensFileEvaluationInputs EvalInputs = CalibrationStepsController->GetLensFileEvaluationInputs();
+	if (!EvalInputs.bIsValid)
 	{
 		return;
 	}
 
-	const ELensDataCategory SelectedCategoryValue = SelectedCategory.GetValue();
-	const FLensFileEvalData& LensFileEvalData = CameraController->GetLensFileEvalDataRef();
-
 	if (SelectedCategoryValue == ELensDataCategory::Focus)
 	{
-		const float RawFocus = LensFileEvalData.Input.Focus;
+		const float RawFocus = EvalInputs.Focus;
 		CommitPositionChange(RawFocus, TEXT("Focus"));
 	}
 	else if (SelectedCategoryValue == ELensDataCategory::Zoom ||
@@ -238,8 +235,8 @@ void FCameraCalibrationTimeSliderController::Tick(float DeltaTime)
 		if (SelectedFocus.IsSet())
 		{
 			
-			const float RawFocus = LensFileEvalData.Input.Focus;
-			const float RawZoom = LensFileEvalData.Input.Zoom;
+			const float RawFocus = EvalInputs.Focus;
+			const float RawZoom = EvalInputs.Zoom;
 			if (LensFileWeakPtr->GetDataTable(SelectedCategoryValue)->IsFocusBetweenNeighbor(SelectedFocus.GetValue(), RawFocus))
 			{
 				CommitPositionChange(RawZoom, TEXT("Zoom"));
@@ -248,7 +245,7 @@ void FCameraCalibrationTimeSliderController::Tick(float DeltaTime)
 	}
 	else if (SelectedCategory == ELensDataCategory::Iris)
 	{
-		const float RawIris = LensFileEvalData.Input.Iris;
+		const float RawIris = EvalInputs.Iris;
 		CommitPositionChange(RawIris, TEXT("Iris"));
 	}
 }
diff --git a/Engine/Plugins/VirtualProduction/CameraCalibration/Source/CameraCalibrationEditor/Private/AssetEditor/CameraCalibrationToolkit.h b/Engine/Plugins/VirtualProduction/CameraCalibration/Source/CameraCalibrationEditor/Private/AssetEditor/CameraCalibrationToolkit.h
--- a/Engine/Plugins/VirtualProduction/CameraCalibration/Source/CameraCalibrationEditor/Private/AssetEditor/CameraCalibrationToolkit.h
+++ b/Engine/Plugins/VirtualProduction/CameraCalibration/Source/CameraCalibrationEditor/Private/AssetEditor/CameraCalibrationToolkit.h
@@ -90,7 +90,7 @@ private:
 	/** Camera calibration tools */
 	TSharedPtr<SWidget> CalibrationStepsTab;
 
-	/** Data displaying evaluated data based on LiveLink inputs */
+	/** Data displaying evaluated data based on evaluation inputs */
 	TSharedPtr<SLensEvaluation> LensEvaluationWidget;
 
 	/** Calibration Steps Controller */
diff --git a/Engine/Plugins/VirtualProduction/CameraCalibration/Source/CameraCalibrationEditor/Private/AssetEditor/ImageCenterTool.h b/Engine/Plugins/VirtualProduction/CameraCalibration/Source/CameraCalibrationEditor/Private/AssetEditor/ImageCenterTool.h
--- a/Engine/Plugins/VirtualProduction/CameraCalibration/Source/CameraCalibrationEditor/Private/AssetEditor/ImageCenterTool.h
+++ b/Engine/Plugins/VirtualProduction/CameraCalibration/Source/CameraCalibrationEditor/Private/AssetEditor/ImageCenterTool.h
@@ -6,8 +6,6 @@
 
 #include "ImageCenterTool.generated.h"
 
-struct FLensFileEvalData;
-
 class UCameraImageCenterAlgo;
 
 /**
diff --git a/Engine/Plugins/VirtualProduction/CameraCalibration/Source/CameraCalibrationEditor/Private/AssetEditor/LensDistortionTool.h b/Engine/Plugins/VirtualProduction/CameraCalibration/Source/CameraCalibrationEditor/Private/AssetEditor/LensDistortionTool.h
--- a/Engine/Plugins/VirtualProduction/CameraCalibration/Source/CameraCalibrationEditor/Private/AssetEditor/LensDistortionTool.h
+++ b/Engine/Plugins/VirtualProduction/CameraCalibration/Source/CameraCalibrationEditor/Private/AssetEditor/LensDistortionTool.h
@@ -11,7 +11,6 @@
 #include "LensDistortionTool.generated.h"
 
 struct FGeometry;
-struct FLensFileEvalData;
 struct FPointerEvent;
 
 class FJsonObject;
diff --git a/Engine/Plugins/VirtualProduction/CameraCalibration/Source/CameraCalibrationEditor/Private/AssetEditor/NodalOffsetTool.cpp b/Engine/Plugins/VirtualProduction/CameraCalibration/Source/CameraCalibrationEditor/Private/AssetEditor/NodalOffsetTool.cpp
--- a/Engine/Plugins/VirtualProduction/CameraCalibration/Source/CameraCalibrationEditor/Private/AssetEditor/NodalOffsetTool.cpp
+++ b/Engine/Plugins/VirtualProduction/CameraCalibration/Source/CameraCalibrationEditor/Private/AssetEditor/NodalOffsetTool.cpp
@@ -250,7 +250,7 @@ void UNodalOffsetTool::OnSaveCurrentNodalOffset()
 	LensFile->Modify();
 	LensFile->AddNodalOffsetPoint(Focus, Zoom, NodalOffset);
 
-	// Force bApplyNodalOffset in the LiveLinkCameraController so that we can see the effect right away
+	// Force the lens component to apply nodal offset so that we can see the effect right away
 	TInlineComponentArray<ULensComponent*> LensComponents;
 	CameraCalibrationStepsController.Pin()->GetCamera()->GetComponents(LensComponents);
 
diff --git a/Engine/Plugins/VirtualProduction/CameraCalibration/Source/CameraCalibrationEditor/Private/AssetEditor/NodalOffsetTool.h b/Engine/Plugins/VirtualProduction/CameraCalibration/Source/CameraCalibrationEditor/Private/AssetEditor/NodalOffsetTool.h
--- a/Engine/Plugins/VirtualProduction/CameraCalibration/Source/CameraCalibrationEditor/Private/AssetEditor/NodalOffsetTool.h
+++ b/Engine/Plugins/VirtualProduction/CameraCalibration/Source/CameraCalibrationEditor/Private/AssetEditor/NodalOffsetTool.h
@@ -11,7 +11,6 @@
 #include "NodalOffsetTool.generated.h"
 
 struct FGeometry;
-struct FLensFileEvalData;
 struct FPointerEvent;
 
 class FJsonObject;
diff --git a/Engine/Plugins/VirtualProduction/CameraCalibration/Source/CameraCalibrationEditor/Private/AssetEditor/SLensDataAddPointDialog.cpp b/Engine/Plugins/VirtualProduction/CameraCalibration/Source/CameraCalibrationEditor/Private/AssetEditor/SLensDataAddPointDialog.cpp
--- a/Engine/Plugins/VirtualProduction/CameraCalibration/Source/CameraCalibrationEditor/Private/AssetEditor/SLensDataAddPointDialog.cpp
+++ b/Engine/Plugins/VirtualProduction/CameraCalibration/Source/CameraCalibrationEditor/Private/AssetEditor/SLensDataAddPointDialog.cpp
@@ -698,7 +698,7 @@ void SLensDataAddPointDialog::SetTrackingData(float Value, int32 Index)
 
 void SLensDataAddPointDialog::RefreshEvaluationData()
 {
-	//Get FIZ coming from LiveLink. If normalized values are available, take them. Otherwise, look for precalibrated values or default to 0.0f
+	//Get FIZ coming from the lens component. If normalized values are available, take them. Otherwise, look for precalibrated values or default to 0.0f
 	const FCachedFIZData& CachedData = CachedFIZ.Get();
 	const float Focus = CachedData.RawFocus.IsSet() ? CachedData.RawFocus.GetValue() : CachedData.EvaluatedFocus.IsSet() ? CachedData.EvaluatedFocus.GetValue() : 0.0f;
 	const float Iris = CachedData.RawIris.IsSet() ? CachedData.RawIris.GetValue() : CachedData.EvaluatedIris.IsSet() ? CachedData.EvaluatedIris.GetValue() : 0.0f;
diff --git a/Engine/Plugins/VirtualProduction/CameraCalibration/Source/CameraCalibrationEditor/Private/AssetEditor/SLensDataViewer.cpp b/Engine/Plugins/VirtualProduction/CameraCalibration/Source/CameraCalibrationEditor/Private/AssetEditor/SLensDataViewer.cpp
--- a/Engine/Plugins/VirtualProduction/CameraCalibration/Source/CameraCalibrationEditor/Private/AssetEditor/SLensDataViewer.cpp
+++ b/Engine/Plugins/VirtualProduction/CameraCalibration/Source/CameraCalibrationEditor/Private/AssetEditor/SLensDataViewer.cpp
@@ -12,7 +12,6 @@
 #include "ICurveEditorModule.h"
 #include "ISinglePropertyView.h"
 #include "LensFile.h"
-#include "LiveLinkCameraController.h"
 #include "RichCurveEditorModel.h"
 #include "SLensDataAddPointDialog.h"
 #include "SLensDataCategoryListItem.h"
diff --git a/Engine/Plugins/VirtualProduction/CameraCalibration/Source/CameraCalibrationEditor/Private/AssetEditor/SLensEvaluation.cpp b/Engine/Plugins/VirtualProduction/CameraCalibration/Source/CameraCalibrationEditor/Private/AssetEditor/SLensEvaluation.cpp
--- a/Engine/Plugins/VirtualProduction/CameraCalibration/Source/CameraCalibrationEditor/Private/AssetEditor/SLensEvaluation.cpp
+++ b/Engine/Plugins/VirtualProduction/CameraCalibration/Source/CameraCalibrationEditor/Private/AssetEditor/SLensEvaluation.cpp
@@ -5,11 +5,8 @@
 #include "CameraCalibrationStepsController.h"
 #include "EngineUtils.h"
 #include "Features/IModularFeatures.h"
-#include "ILiveLinkClient.h"
+#include "LensComponent.h"
 #include "LensFile.h"
-#include "LiveLinkCameraController.h"
-#include "Roles/LiveLinkCameraRole.h"
-#include "Roles/LiveLinkCameraTypes.h"
 #include "UI/CameraCalibrationEditorStyle.h"
 #include "Widgets/Input/SCheckBox.h"
 #include "Widgets/Layout/SGridPanel.h"
@@ -85,64 +82,58 @@ void SLensEvaluation::Construct(const FArguments& InArgs, TWeakPtr<FCameraCalibr
 
 void SLensEvaluation::Tick(const FGeometry& AllottedGeometry, const double InCurrentTime, const float InDeltaTime)
 {
-	//Cache LiveLink data every tick to be sure we have the right one for the frame during calibration
-	CacheLiveLinkData();
+	//Cache LensFile evaluation data every tick to be sure we have the right one for the frame during calibration
+	CacheLensFileEvaluationInputs();
 	CacheLensFileData();
 
 	Super::Tick(AllottedGeometry, InCurrentTime, InDeltaTime);
 }
 
-void SLensEvaluation::CacheLiveLinkData()
+void SLensEvaluation::CacheLensFileEvaluationInputs()
 {
 	//Start clean
-	CachedLiveLinkData.RawFocus.Reset();
-	CachedLiveLinkData.RawIris.Reset();
-	CachedLiveLinkData.RawZoom.Reset();
-	CachedLiveLinkData.EvaluatedFocus.Reset();
-	CachedLiveLinkData.EvaluatedIris.Reset();
-	CachedLiveLinkData.EvaluatedZoom.Reset();
-
-	if (ULiveLinkCameraController* CameraController = WeakStepsController.Pin()->FindLiveLinkCameraController())
+	CachedFIZData.RawFocus.Reset();
+	CachedFIZData.RawIris.Reset();
+	CachedFIZData.RawZoom.Reset();
+	CachedFIZData.EvaluatedFocus.Reset();
+	CachedFIZData.EvaluatedIris.Reset();
+	CachedFIZData.EvaluatedZoom.Reset();
+
+	if (TSharedPtr<FCameraCalibrationStepsController> StepsController = WeakStepsController.Pin())
 	{
-		bCameraControllerExists = true;
-
-		const FLensFileEvalData& LensFileEvalData = CameraController->GetLensFileEvalDataRef();
+		const FLensFileEvaluationInputs EvalInputs = StepsController->GetLensFileEvaluationInputs();
+		if (!EvalInputs.bIsValid)
 		{
-			CachedLiveLinkData.RawFocus = LensFileEvalData.Input.Focus;
-			if (LensFile->HasFocusEncoderMapping())
-			{
-				CachedLiveLinkData.EvaluatedFocus = LensFile->EvaluateNormalizedFocus(LensFileEvalData.Input.Focus);
-			}
+			return;
 		}
+
+		CachedFIZData.RawFocus = EvalInputs.Focus;
+		if (LensFile->HasFocusEncoderMapping())
 		{
-			CachedLiveLinkData.RawIris = LensFileEvalData.Input.Iris;
-			if (LensFile->HasIrisEncoderMapping())
-			{
-				CachedLiveLinkData.EvaluatedIris = LensFile->EvaluateNormalizedIris(LensFileEvalData.Input.Iris);
-			}
+			CachedFIZData.EvaluatedFocus = LensFile->EvaluateNormalizedFocus(EvalInputs.Focus);
 		}
+
+		CachedFIZData.RawIris = EvalInputs.Iris;
+		if (LensFile->HasIrisEncoderMapping())
 		{
-			CachedLiveLinkData.RawZoom = LensFileEvalData.Input.Zoom;
+			CachedFIZData.EvaluatedIris = LensFile->EvaluateNormalizedIris(EvalInputs.Iris);
+		}
 
-			FFocalLengthInfo FocalLength;
-			if (LensFile->EvaluateFocalLength(LensFileEvalData.Input.Focus, LensFileEvalData.Input.Zoom, FocalLength))
-			{
-				CachedLiveLinkData.EvaluatedZoom = FocalLength.FxFy.X * LensFile->LensInfo.SensorDimensions.X;
-			}
+		CachedFIZData.RawZoom = EvalInputs.Zoom;
+		FFocalLengthInfo FocalLength;
+		if (LensFile->EvaluateFocalLength(EvalInputs.Focus, EvalInputs.Zoom, FocalLength))
+		{
+			CachedFIZData.EvaluatedZoom = FocalLength.FxFy.X * LensFile->LensInfo.SensorDimensions.X;
 		}
 	}
-	else
-	{
-		bCameraControllerExists = false;
-	}
 }
 
 void SLensEvaluation::CacheLensFileData()
 {
-	//Evaluate LensFile independantly of valid FZ pair. Use default 0.0f like LiveLinkCamera if it's not present
+	//Evaluate LensFile independently of valid FZ pair. Use default 0.0f if it's not present
 	{
-		const float Focus = CachedLiveLinkData.RawFocus.IsSet() ? CachedLiveLinkData.RawFocus.GetValue() : 0.0f;
-		const float Zoom = CachedLiveLinkData.RawZoom.IsSet() ? CachedLiveLinkData.RawZoom.GetValue() : 0.0f;
+		const float Focus = CachedFIZData.RawFocus.IsSet() ? CachedFIZData.RawFocus.GetValue() : 0.0f;
+		const float Zoom = CachedFIZData.RawZoom.IsSet() ? CachedFIZData.RawZoom.GetValue() : 0.0f;
 		bCouldEvaluateDistortion = LensFile->EvaluateDistortionParameters(Focus, Zoom, CachedDistortionInfo);
 		bCouldEvaluateFocalLength = LensFile->EvaluateFocalLength(Focus, Zoom, CachedFocalLengthInfo);
 		bCouldEvaluateImageCenter = LensFile->EvaluateImageCenterParameters(Focus, Zoom, CachedImageCenter);
@@ -178,15 +169,15 @@ TSharedRef<SWidget> SLensEvaluation::MakeTrackingWidget()
 			.Padding(0.0f, 15.0f, 0.0f, 5.0f)
 			[
 				SNew(STextBlock)
-				.Text(LOCTEXT("LiveLinkLabelSection", "Selected LiveLink Subject"))
+				.Text(LOCTEXT("LensComponentLabelSection", "Lens Component"))
 				.Font(FAppStyle::GetFontStyle("DetailsView.CategoryFontStyle"))
 				.ShadowOffset(FVector2D(1.0f, 1.0f))
 			]
 			+ SGridPanel::Slot(0, 3)
 			[
 				SNew(STextBlock)
-				.Text(this, &SLensEvaluation::GetLiveLinkCameraControllerLabel)
-				.ColorAndOpacity(this, &SLensEvaluation::GetLiveLinkCameraControllerLabelColor)
+				.Text(this, &SLensEvaluation::GetLensComponentLabel)
+				.ColorAndOpacity(this, &SLensEvaluation::GetLensComponentLabelColor)
 				.AutoWrapText(true)
 			]
 		];
@@ -217,23 +208,23 @@ FSlateColor SLensEvaluation::GetTrackedCameraLabelColor() const
 }
 
 
-FText SLensEvaluation::GetLiveLinkCameraControllerLabel() const
+FText SLensEvaluation::GetLensComponentLabel() const
 {
 	if (WeakStepsController.IsValid())
 	{
-		if (ULiveLinkCameraController* CameraController = WeakStepsController.Pin()->FindLiveLinkCameraController())
+		if (ULensComponent* LensComponent = WeakStepsController.Pin()->FindLensComponent())
 		{
-			return FText::FromName(CameraController->GetSelectedSubject().Subject.Name);
+			return FText::FromName(LensComponent->GetFName());
 		}
 	}
-	return LOCTEXT("NoLiveLinkSubject", "The tracked camera does not have a LiveLink camera controller, or the controller's LensFile does not match this one.");
+	return LOCTEXT("NoLensComponent", "The tracked camera does not have a Lens Component, or the LensFile assigned in that Lens Component does not match this one.");
 }
 
-FSlateColor SLensEvaluation::GetLiveLinkCameraControllerLabelColor() const
+FSlateColor SLensEvaluation::GetLensComponentLabelColor() const
 {
 	if (WeakStepsController.IsValid())
 	{
-		if (ULiveLinkCameraController* CameraController = WeakStepsController.Pin()->FindLiveLinkCameraController())
+		if (ULensComponent* LensComponent = WeakStepsController.Pin()->FindLensComponent())
 		{
 			return FLinearColor::White;
 		}
@@ -286,15 +277,15 @@ TSharedRef<SWidget> SLensEvaluation::MakeRawInputFIZWidget() const
 				SNew(STextBlock)
 				.Text(MakeAttributeLambda([this]
 				{
-					if (CachedLiveLinkData.RawFocus.IsSet())
+					if (CachedFIZData.RawFocus.IsSet())
 					{
-						return FText::AsNumber(CachedLiveLinkData.RawFocus.GetValue());
+						return FText::AsNumber(CachedFIZData.RawFocus.GetValue());
 					}
-					return LOCTEXT("NoRawFocus", "No Focus From LiveLink");
+					return LOCTEXT("NoRawFocus", "No Focus Input");
 				}))
 				.ColorAndOpacity(MakeAttributeLambda([this]
 				{
-					if (CachedLiveLinkData.RawFocus.IsSet())
+					if (CachedFIZData.RawFocus.IsSet())
 					{
 						return FSlateColor(FLinearColor::White);
 					}
@@ -306,15 +297,15 @@ TSharedRef<SWidget> SLensEvaluation::MakeRawInputFIZWidget() const
 				SNew(STextBlock)
 				.Text(MakeAttributeLambda([this]
 				{
-					if (CachedLiveLinkData.RawIris.IsSet())
+					if (CachedFIZData.RawIris.IsSet())
 					{
-						return FText::AsNumber(CachedLiveLinkData.RawIris.GetValue());
+						return FText::AsNumber(CachedFIZData.RawIris.GetValue());
 					}
-					return LOCTEXT("NoRawIris", "No Iris From LiveLink");
+					return LOCTEXT("NoRawIris", "No Iris Input");
 				}))
 				.ColorAndOpacity(MakeAttributeLambda([this]
 				{
-					if (CachedLiveLinkData.RawIris.IsSet())
+					if (CachedFIZData.RawIris.IsSet())
 					{
 						return FSlateColor(FLinearColor::White);
 					}
@@ -326,15 +317,15 @@ TSharedRef<SWidget> SLensEvaluation::MakeRawInputFIZWidget() const
 				SNew(STextBlock)
 				.Text(MakeAttributeLambda([this]
 				{
-					if (CachedLiveLinkData.RawZoom.IsSet())
+					if (CachedFIZData.RawZoom.IsSet())
 					{
-						return FText::AsNumber(CachedLiveLinkData.RawZoom.GetValue());
+						return FText::AsNumber(CachedFIZData.RawZoom.GetValue());
 					}
-					return LOCTEXT("NoRawZoom", "No Zoom From LiveLink");
+					return LOCTEXT("NoRawZoom", "No Zoom Input");
 				}))
 				.ColorAndOpacity(MakeAttributeLambda([this]
 				{
-					if (CachedLiveLinkData.RawZoom.IsSet())
+					if (CachedFIZData.RawZoom.IsSet())
 					{
 						return FSlateColor(FLinearColor::White);
 					}
@@ -400,15 +391,11 @@ TSharedRef<SWidget> SLensEvaluation::MakeEvaluatedFIZWidget() const
 				SNew(STextBlock)
 				.Text(MakeAttributeLambda([this, FloatOptions]
 				{
-					if (CachedLiveLinkData.EvaluatedFocus.IsSet())
+					if (CachedFIZData.EvaluatedFocus.IsSet())
 					{
-						const FText ValueText = FText::AsNumber(CachedLiveLinkData.EvaluatedFocus.GetValue(), &FloatOptions);
+						const FText ValueText = FText::AsNumber(CachedFIZData.EvaluatedFocus.GetValue(), &FloatOptions);
 						return FText::Format(LOCTEXT("PhysicalUnitsFocusValue", "{0} cm"), ValueText);
 					}
-					else if (!bCameraControllerExists)
-					{
-						return LOCTEXT("UndefinedValue", "N/A");
-					}
 					else if (!LensFile->HasSamples(ELensDataCategory::Focus))
 					{
 						return LOCTEXT("NoFocusCurve", "No Focus Curve");
@@ -417,11 +404,7 @@ TSharedRef<SWidget> SLensEvaluation::MakeEvaluatedFIZWidget() const
 				}))
 				.ColorAndOpacity(MakeAttributeLambda([this]
 				{
-					if (CachedLiveLinkData.EvaluatedFocus.IsSet())
-					{
-						return FSlateColor(FLinearColor::White);
-					}
-					else if (!bCameraControllerExists)
+					if (CachedFIZData.EvaluatedFocus.IsSet())
 					{
 						return FSlateColor(FLinearColor::White);
 					}
@@ -437,18 +420,14 @@ TSharedRef<SWidget> SLensEvaluation::MakeEvaluatedFIZWidget() const
 				SNew(STextBlock)
 				.Text(MakeAttributeLambda([this]
 				{
-					if (CachedLiveLinkData.EvaluatedIris.IsSet())
+					if (CachedFIZData.EvaluatedIris.IsSet())
 					{
 						FNumberFormattingOptions IrisOptions;
 						IrisOptions.MinimumFractionalDigits = 1;
 						IrisOptions.MaximumFractionalDigits = 1;
-						const FText ValueText = FText::AsNumber(CachedLiveLinkData.EvaluatedIris.GetValue(), &IrisOptions);
+						const FText ValueText = FText::AsNumber(CachedFIZData.EvaluatedIris.GetValue(), &IrisOptions);
 						return FText::Format(LOCTEXT("PhysicalUnitsIrisValue", "{0} F-Stop"), ValueText);
 					}
-					else if (!bCameraControllerExists)
-					{
-						return LOCTEXT("UndefinedValue", "N/A");
-					}
 					else if (!LensFile->HasSamples(ELensDataCategory::Iris))
 					{
 						return LOCTEXT("NoIrisCurve", "No Iris Curve");
@@ -457,11 +436,7 @@ TSharedRef<SWidget> SLensEvaluation::MakeEvaluatedFIZWidget() const
 				}))
 				.ColorAndOpacity(MakeAttributeLambda([this]
 				{
-					if (CachedLiveLinkData.EvaluatedIris.IsSet())
-					{
-						return FSlateColor(FLinearColor::White);
-					}
-					else if (!bCameraControllerExists)
+					if (CachedFIZData.EvaluatedIris.IsSet())
 					{
 						return FSlateColor(FLinearColor::White);
 					}
@@ -477,15 +452,11 @@ TSharedRef<SWidget> SLensEvaluation::MakeEvaluatedFIZWidget() const
 				SNew(STextBlock)
 				.Text(MakeAttributeLambda([this, FloatOptions]
 				{
-					if (CachedLiveLinkData.EvaluatedZoom.IsSet())
+					if (CachedFIZData.EvaluatedZoom.IsSet())
 					{
-						const FText ValueText = FText::AsNumber(CachedLiveLinkData.EvaluatedZoom.GetValue(), &FloatOptions);
+						const FText ValueText = FText::AsNumber(CachedFIZData.EvaluatedZoom.GetValue(), &FloatOptions);
 						return FText::Format(LOCTEXT("PhysicalUnitsZoomValue", "{0} mm"), ValueText);
 					}
-					else if (!bCameraControllerExists)
-					{
-						return LOCTEXT("UndefinedValue", "N/A");
-					}
 					else if (!LensFile->HasSamples(ELensDataCategory::Zoom))
 					{
 						return LOCTEXT("NoZoomCurve", "No Focal Length Curve");
@@ -494,11 +465,7 @@ TSharedRef<SWidget> SLensEvaluation::MakeEvaluatedFIZWidget() const
 				}))
 				.ColorAndOpacity(MakeAttributeLambda([this]
 				{
-					if (CachedLiveLinkData.EvaluatedZoom.IsSet())
-					{
-						return FSlateColor(FLinearColor::White);
-					}
-					else if (!bCameraControllerExists)
+					if (CachedFIZData.EvaluatedZoom.IsSet())
 					{
 						return FSlateColor(FLinearColor::White);
 					}
@@ -514,9 +481,9 @@ TSharedRef<SWidget> SLensEvaluation::MakeEvaluatedFIZWidget() const
 				SNew(STextBlock)
 				.Text(MakeAttributeLambda([this, FloatOptions]
 				{
-					if (CachedLiveLinkData.EvaluatedZoom.IsSet())
+					if (CachedFIZData.EvaluatedZoom.IsSet())
 					{
-						const float FOV = FMath::RadiansToDegrees(2.f * FMath::Atan(LensFile->LensInfo.SensorDimensions.X / (2.f * CachedLiveLinkData.EvaluatedZoom.GetValue())));
+						const float FOV = FMath::RadiansToDegrees(2.f * FMath::Atan(LensFile->LensInfo.SensorDimensions.X / (2.f * CachedFIZData.EvaluatedZoom.GetValue())));
 						const FText ValueText = FText::AsNumber(FOV, &FloatOptions);
 						return FText::Format(LOCTEXT("PhysicalUnitsFOVValue", "{0} deg"), ValueText);
 					}
diff --git a/Engine/Plugins/VirtualProduction/CameraCalibration/Source/CameraCalibrationEditor/Private/AssetEditor/SLensEvaluation.h b/Engine/Plugins/VirtualProduction/CameraCalibration/Source/CameraCalibrationEditor/Private/AssetEditor/SLensEvaluation.h
--- a/Engine/Plugins/VirtualProduction/CameraCalibration/Source/CameraCalibrationEditor/Private/AssetEditor/SLensEvaluation.h
+++ b/Engine/Plugins/VirtualProduction/CameraCalibration/Source/CameraCalibrationEditor/Private/AssetEditor/SLensEvaluation.h
@@ -5,9 +5,7 @@
 #include "Widgets/SCompoundWidget.h"
 
 #include "Camera/CameraActor.h"
-#include "LiveLinkRole.h"
 #include "SLensFilePanel.h"
-#include "SLiveLinkSubjectRepresentationPicker.h"
 #include "UObject/StrongObjectPtr.h"
 #include "Widgets/DeclarativeSyntaxSupport.h"
 
@@ -16,7 +14,7 @@ class ULensFile;
 class FCameraCalibrationStepsController;
 
 /**
- * Widget using LiveLink subject input to evaluate lens file and show resulting data
+ * Widget using lens file evaluation inputs to evaluate lens file and show resulting data
  */
 class SLensEvaluation : public SCompoundWidget
 {
@@ -32,13 +30,13 @@ public:
 	virtual void Tick(const FGeometry& AllottedGeometry, const double InCurrentTime, const float InDeltaTime) override;
 	//~ End SCompoundWidget interface
 
-	/** Returns last evaluated FIZ data from LiveLink */
-	FCachedFIZData GetLastEvaluatedData() const { return CachedLiveLinkData; }
+	/** Returns last raw and evaluated FIZ data */
+	FCachedFIZData GetLastEvaluatedData() const { return CachedFIZData; }
 
 private:
 
-	/** Evaluates LiveLink subjects */
-	void CacheLiveLinkData();
+	/** Caches latest evaluation inputs and evaluates LensFile for FIZ */
+	void CacheLensFileEvaluationInputs();
 
 	/** Evaluates LensFile using tracking data */
 	void CacheLensFileData();
@@ -49,19 +47,19 @@ private:
 	/** Get the text color for the tracked camera text */
 	FSlateColor GetTrackedCameraLabelColor() const;
 
-	/** Get the name of the livelink camera controller driving the selected camera */
-	FText GetLiveLinkCameraControllerLabel() const;
+	/** Get the name of the lens component driving the selected camera */
+	FText GetLensComponentLabel() const;
 
-	/** Get the text color for the livelink camera controller text */
-	FSlateColor GetLiveLinkCameraControllerLabelColor() const;
+	/** Get the text color for the lens component text */
+	FSlateColor GetLensComponentLabelColor() const;
 
-	/** Make LiveLink subject selection widget */
+	/** Make widget to display selected camera and lens component */
 	TSharedRef<SWidget> MakeTrackingWidget();
 
-	/** Make FIZ widget evaluated from LiveLink subject */
+	/** Make Raw FIZ widget */
 	TSharedRef<SWidget> MakeRawInputFIZWidget() const;
 
-	/** Make FIZ widget evaluated from LiveLink subject */
+	/** Make Evaluated FIZ widget */
 	TSharedRef<SWidget> MakeEvaluatedFIZWidget() const;
 
 	/** Make widget showing distortion data evaluated from LensFile */
@@ -78,20 +76,17 @@ private:
 	/** LensFile being edited */
 	TStrongObjectPtr<ULensFile> LensFile;
 
-	/** Whether the tracked camera has a LiveLink Camera Controller */
-	bool bCameraControllerExists = false;
-
 	/** Flags used to know what data is valid or not and adjust UI */
 	bool bCouldEvaluateDistortion = false;
 	bool bCouldEvaluateFocalLength = false;
 	bool bCouldEvaluateImageCenter = false;
 	bool bCouldEvaluateNodalOffset = false;
 
-	/** Calibration steps controller, which provides a selected target camera and livelink camera controller */
+	/** Calibration steps controller */
 	TWeakPtr<FCameraCalibrationStepsController> WeakStepsController;
 
-	//~ Cached data taken from LiveLink subject and LensFile
-	FCachedFIZData CachedLiveLinkData;
+	//~ Cached LensFile evaluated data
+	FCachedFIZData CachedFIZData;
 	FDistortionInfo CachedDistortionInfo;
 	FFocalLengthInfo CachedFocalLengthInfo;
 	FImageCenterInfo CachedImageCenter;
diff --git a/Engine/Plugins/VirtualProduction/CameraCalibration/Source/CameraCalibrationEditor/Private/Calibrators/CameraImageCenterAlgoManual.cpp b/Engine/Plugins/VirtualProduction/CameraCalibration/Source/CameraCalibrationEditor/Private/Calibrators/CameraImageCenterAlgoManual.cpp
--- a/Engine/Plugins/VirtualProduction/CameraCalibration/Source/CameraCalibrationEditor/Private/Calibrators/CameraImageCenterAlgoManual.cpp
+++ b/Engine/Plugins/VirtualProduction/CameraCalibration/Source/CameraCalibrationEditor/Private/Calibrators/CameraImageCenterAlgoManual.cpp
@@ -61,15 +61,15 @@ void UCameraImageCenterAlgoManual::Activate()
 
 	RenderTargetSize = FIntPoint(RenderTarget->SizeX, RenderTarget->SizeY);
 
-	const FLensFileEvalData* EvalData = StepsController->GetLensFileEvalData();
-	if (!EvalData)
+	const FLensFileEvaluationInputs EvalInputs = StepsController->GetLensFileEvaluationInputs();
+	if (!EvalInputs.bIsValid)
 	{
 		return;
 	}
 
 	// Cache the original image center point for the current focus and zoom
-	CurrentEvalFocus = EvalData->Input.Focus;
-	CurrentEvalZoom = EvalData->Input.Zoom;
+	CurrentEvalFocus = EvalInputs.Focus;
+	CurrentEvalZoom = EvalInputs.Zoom;
 	LensFile->EvaluateImageCenterParameters(CurrentEvalFocus, CurrentEvalZoom, OriginalImageCenter);
 	AdjustedImageCenter = OriginalImageCenter;
 
@@ -94,17 +94,18 @@ void UCameraImageCenterAlgoManual::Tick(float DeltaTime)
 
 	if (LensFile.IsValid())
 	{
-		if (const FLensFileEvalData* EvalData = StepsController->GetLensFileEvalData())
+		const FLensFileEvaluationInputs EvalInputs = StepsController->GetLensFileEvaluationInputs();
+		if (EvalInputs.bIsValid)
 		{
 			// Check if the focus or zoom has changed since the last tick
 			const float InputTolerance = GetDefault<UCameraCalibrationSettings>()->GetCalibrationInputTolerance();
-			if (!FMath::IsNearlyEqual(CurrentEvalFocus, EvalData->Input.Focus, InputTolerance) || !FMath::IsNearlyEqual(CurrentEvalZoom, EvalData->Input.Zoom, InputTolerance))
+			if (!FMath::IsNearlyEqual(CurrentEvalFocus, EvalInputs.Focus, InputTolerance) || !FMath::IsNearlyEqual(CurrentEvalZoom, EvalInputs.Zoom, InputTolerance))
 			{
 				// Prompt the user to save any unsaved changes before caching a new original image center based on the new focus/zoom
 				ApplyOrRevertAdjustedImageCenter();
 
-				CurrentEvalFocus = EvalData->Input.Focus;
-				CurrentEvalZoom = EvalData->Input.Zoom;
+				CurrentEvalFocus = EvalInputs.Focus;
+				CurrentEvalZoom = EvalInputs.Zoom;
 				LensFile->EvaluateImageCenterParameters(CurrentEvalFocus, CurrentEvalZoom, OriginalImageCenter);
 				AdjustedImageCenter = OriginalImageCenter;
 			}
diff --git a/Engine/Plugins/VirtualProduction/CameraCalibration/Source/CameraCalibrationEditor/Private/Calibrators/CameraLensDistortionAlgoCheckerboard.cpp b/Engine/Plugins/VirtualProduction/CameraCalibration/Source/CameraCalibrationEditor/Private/Calibrators/CameraLensDistortionAlgoCheckerboard.cpp
--- a/Engine/Plugins/VirtualProduction/CameraCalibration/Source/CameraCalibrationEditor/Private/Calibrators/CameraLensDistortionAlgoCheckerboard.cpp
+++ b/Engine/Plugins/VirtualProduction/CameraCalibration/Source/CameraCalibrationEditor/Private/Calibrators/CameraLensDistortionAlgoCheckerboard.cpp
@@ -156,18 +156,16 @@ void UCameraLensDistortionAlgoCheckerboard::Tick(float DeltaTime)
 		{
 			LastCameraData.bIsValid = false;
 
-			const FLensFileEvalData* LensFileEvalData = StepsController->GetLensFileEvalData();
+			const FLensFileEvaluationInputs LensFileEvalInputs = StepsController->GetLensFileEvaluationInputs();
 
 			// We require lens evaluation data.
-			if (!LensFileEvalData)
+			if (LensFileEvalInputs.bIsValid)
 			{
-				break;
+				LastCameraData.InputFocus = LensFileEvalInputs.Focus;
+				LastCameraData.InputZoom = LensFileEvalInputs.Zoom;
+				LastCameraData.bIsValid = true;
 			}
 
-			LastCameraData.InputFocus = LensFileEvalData->Input.Focus;
-			LastCameraData.InputZoom = LensFileEvalData->Input.Zoom;
-			LastCameraData.bIsValid = true;
-
 		} while (0);
 	}
 }
@@ -230,7 +228,7 @@ bool UCameraLensDistortionAlgoCheckerboard::AddCalibrationRow(FText& OutErrorMes
 
 	if (!LastCameraData.bIsValid)
 	{
-		OutErrorMessage = LOCTEXT("InvalidLastCameraData", "Could not find a cached set of camera data (e.g. FIZ). Please ensure that the camera has a camera live link controller and that it is receiving FIZ data");
+		OutErrorMessage = LOCTEXT("InvalidLastCameraData", "Could not find a cached set of camera data (e.g. FIZ). Check the Lens Component to make sure it has valid evaluation inputs.");
 		return false;
 	}
 
@@ -447,8 +445,6 @@ bool UCameraLensDistortionAlgoCheckerboard::ValidateNewRow(TSharedPtr<FLensDisto
 		return false;
 	}
 	
-	// FZ inputs are always valid, no need to verify them. They could be coming from LiveLink or fallback to a default one
-
 	// Valid image dimensions
 	if ((Row->ImageHeight < 1) || (Row->ImageWidth < 1))
 	{
diff --git a/Engine/Plugins/VirtualProduction/CameraCalibration/Source/CameraCalibrationEditor/Private/Calibrators/CameraLensDistortionAlgoPoints.cpp b/Engine/Plugins/VirtualProduction/CameraCalibration/Source/CameraCalibrationEditor/Private/Calibrators/CameraLensDistortionAlgoPoints.cpp
--- a/Engine/Plugins/VirtualProduction/CameraCalibration/Source/CameraCalibrationEditor/Private/Calibrators/CameraLensDistortionAlgoPoints.cpp
+++ b/Engine/Plugins/VirtualProduction/CameraCalibration/Source/CameraCalibrationEditor/Private/Calibrators/CameraLensDistortionAlgoPoints.cpp
@@ -177,15 +177,14 @@ void UCameraLensDistortionAlgoPoints::Tick(float DeltaTime)
 
 			LastCameraData.UniqueId = Camera->GetUniqueID();
 
-			const FLensFileEvalData* LensFileEvalData = StepsController->GetLensFileEvalData();
-
-			if (!LensFileEvalData)
+			const FLensFileEvaluationInputs EvalInputs = StepsController->GetLensFileEvaluationInputs();
+			if (!EvalInputs.bIsValid)
 			{
 				break;
 			}
 
-			LastCameraData.InputFocus = LensFileEvalData->Input.Focus;
-			LastCameraData.InputZoom = LensFileEvalData->Input.Zoom;
+			LastCameraData.InputFocus = EvalInputs.Focus;
+			LastCameraData.InputZoom = EvalInputs.Zoom;
 
 			if (!Calibrator.IsValid())
 			{
@@ -220,6 +219,13 @@ bool UCameraLensDistortionAlgoPoints::OnViewportClicked(const FGeometry& MyGeome
 		return true;
 	}
 
+	if (!LastCameraData.bIsValid)
+	{
+		FText ErrorMessage = LOCTEXT("InvalidLastCameraData", "Could not find a cached set of camera data (e.g. FIZ). Check the Lens Component to make sure it has valid evaluation inputs.");
+		FMessageDialog::Open(EAppMsgType::Ok, ErrorMessage);
+		return true;
+	}
+
 	// Get currently selected calibrator point
 	FLensDistortionPointsCalibratorPointData LastCalibratorPoint;
 	LastCalibratorPoint.bIsValid = false;
@@ -241,7 +247,7 @@ bool UCameraLensDistortionAlgoPoints::OnViewportClicked(const FGeometry& MyGeome
 	}
 
 	// Check that we have a valid calibrator 3dpoint or camera data
-	if (!LastCalibratorPoint.bIsValid || !LastCameraData.bIsValid)
+	if (!LastCalibratorPoint.bIsValid)
 	{
 		return true;
 	}
@@ -689,8 +695,6 @@ bool UCameraLensDistortionAlgoPoints::ValidateNewRow(TSharedPtr<FLensDistortionP
 		return false;
 	}
 
-	// FZ inputs are always valid, no need to verify them. They could be coming from LiveLink or fallback to a default one
-
 	//@todo Focus and zoom did not change much (i.e. inputs to distortion and nodal offset). 
 	//      Threshold for physical units should differ from normalized encoders.
 
diff --git a/Engine/Plugins/VirtualProduction/CameraCalibration/Source/CameraCalibrationEditor/Private/Calibrators/CameraNodalOffsetAlgoAruco.cpp b/Engine/Plugins/VirtualProduction/CameraCalibration/Source/CameraCalibrationEditor/Private/Calibrators/CameraNodalOffsetAlgoAruco.cpp
--- a/Engine/Plugins/VirtualProduction/CameraCalibration/Source/CameraCalibrationEditor/Private/Calibrators/CameraNodalOffsetAlgoAruco.cpp
+++ b/Engine/Plugins/VirtualProduction/CameraCalibration/Source/CameraCalibrationEditor/Private/Calibrators/CameraNodalOffsetAlgoAruco.cpp
@@ -115,6 +115,12 @@ bool UCameraNodalOffsetAlgoAruco::PopulatePoints(FText& OutErrorMessage)
 		return false;
 	}
 
+	if (!LastCameraData.bIsValid)
+	{
+		OutErrorMessage = LOCTEXT("InvalidLastCameraData", "Could not find a cached set of camera data (e.g. FIZ). Check the Lens Component to make sure it has valid evaluation inputs.");
+		return false;
+	}
+
 	// Detect Aruco dictionary used by looking at the calibration point names
 
 	struct FArucoDictionaryInfo
diff --git a/Engine/Plugins/VirtualProduction/CameraCalibration/Source/CameraCalibrationEditor/Private/Calibrators/CameraNodalOffsetAlgoCheckerboard.cpp b/Engine/Plugins/VirtualProduction/CameraCalibration/Source/CameraCalibrationEditor/Private/Calibrators/CameraNodalOffsetAlgoCheckerboard.cpp
--- a/Engine/Plugins/VirtualProduction/CameraCalibration/Source/CameraCalibrationEditor/Private/Calibrators/CameraNodalOffsetAlgoCheckerboard.cpp
+++ b/Engine/Plugins/VirtualProduction/CameraCalibration/Source/CameraCalibrationEditor/Private/Calibrators/CameraNodalOffsetAlgoCheckerboard.cpp
@@ -120,6 +120,12 @@ bool UCameraNodalOffsetAlgoCheckerboard::PopulatePoints(FText& OutErrorMessage)
 		return false;
 	}
 
+	if (!LastCameraData.bIsValid)
+	{
+		OutErrorMessage = LOCTEXT("InvalidLastCameraData", "Could not find a cached set of camera data (e.g. FIZ). Check the Lens Component to make sure it has valid evaluation inputs.");
+		return false;
+	}
+
 	TArray<FColor> Pixels;
 	FIntPoint Size;
 	ETextureRenderTargetFormat PixelFormat;
diff --git a/Engine/Plugins/VirtualProduction/CameraCalibration/Source/CameraCalibrationEditor/Private/Calibrators/CameraNodalOffsetAlgoOpticalAxis.cpp b/Engine/Plugins/VirtualProduction/CameraCalibration/Source/CameraCalibrationEditor/Private/Calibrators/CameraNodalOffsetAlgoOpticalAxis.cpp
--- a/Engine/Plugins/VirtualProduction/CameraCalibration/Source/CameraCalibrationEditor/Private/Calibrators/CameraNodalOffsetAlgoOpticalAxis.cpp
+++ b/Engine/Plugins/VirtualProduction/CameraCalibration/Source/CameraCalibrationEditor/Private/Calibrators/CameraNodalOffsetAlgoOpticalAxis.cpp
@@ -96,10 +96,11 @@ void UCameraNodalOffsetAlgoOpticalAxis::Tick(float DeltaTime)
 	// Inform the user if the focus or zoom has changed since the last tick
 	if (OpticalAxis.IsSet() && !bFocusZoomWarningIssued)
 	{
-		if (const FLensFileEvalData* EvalData = StepsController->GetLensFileEvalData())
+		const FLensFileEvaluationInputs EvalInputs = StepsController->GetLensFileEvaluationInputs();
+		if (EvalInputs.bIsValid)
 		{
 			const float InputTolerance = GetDefault<UCameraCalibrationSettings>()->GetCalibrationInputTolerance();
-			if (!FMath::IsNearlyEqual(CachedFocus, EvalData->Input.Focus, InputTolerance) || !FMath::IsNearlyEqual(CachedZoom, EvalData->Input.Zoom, InputTolerance))
+			if (!FMath::IsNearlyEqual(CachedFocus, EvalInputs.Focus, InputTolerance) || !FMath::IsNearlyEqual(CachedZoom, EvalInputs.Zoom, InputTolerance))
 			{
 				FFormatOrderedArguments Arguments;
 				Arguments.Add(FText::FromString(FString::Printf(TEXT("%.2f"), CachedFocus)));
diff --git a/Engine/Plugins/VirtualProduction/CameraCalibration/Source/CameraCalibrationEditor/Private/Calibrators/CameraNodalOffsetAlgoPoints.cpp b/Engine/Plugins/VirtualProduction/CameraCalibration/Source/CameraCalibrationEditor/Private/Calibrators/CameraNodalOffsetAlgoPoints.cpp
--- a/Engine/Plugins/VirtualProduction/CameraCalibration/Source/CameraCalibrationEditor/Private/Calibrators/CameraNodalOffsetAlgoPoints.cpp
+++ b/Engine/Plugins/VirtualProduction/CameraCalibration/Source/CameraCalibrationEditor/Private/Calibrators/CameraNodalOffsetAlgoPoints.cpp
@@ -593,10 +593,10 @@ void UCameraNodalOffsetAlgoPoints::Tick(float DeltaTime)
 		{
 			LastCameraData.bIsValid = false;
 
-			const FLensFileEvalData* LensFileEvalData = StepsController->GetLensFileEvalData();
+			const FLensFileEvaluationInputs EvalInputs = StepsController->GetLensFileEvaluationInputs();
 
 			// We require lens evaluation data
-			if (!LensFileEvalData)
+			if (!EvalInputs.bIsValid)
 			{
 				break;
 			}
@@ -617,8 +617,8 @@ void UCameraNodalOffsetAlgoPoints::Tick(float DeltaTime)
 
 			LastCameraData.Pose = CameraComponent->GetComponentToWorld();
 			LastCameraData.UniqueId = Camera->GetUniqueID();
-			LastCameraData.InputFocus = LensFileEvalData->Input.Focus;
-			LastCameraData.InputZoom = LensFileEvalData->Input.Zoom;
+			LastCameraData.InputFocus = EvalInputs.Focus;
+			LastCameraData.InputZoom = EvalInputs.Zoom;
 
 			const ULensComponent* LensComponent = StepsController->FindLensComponent();
 			if (LensComponent)
@@ -701,6 +701,13 @@ bool UCameraNodalOffsetAlgoPoints::OnViewportClicked(const FGeometry& MyGeometry
 		return true;
 	}
 
+	if (!LastCameraData.bIsValid)
+	{
+		FText ErrorMessage = LOCTEXT("InvalidLastCameraData", "Could not find a cached set of camera data (e.g. FIZ). Check the Lens Component to make sure it has valid evaluation inputs.");
+		FMessageDialog::Open(EAppMsgType::Ok, ErrorMessage);
+		return true;
+	}
+
 	// Get currently selected calibrator point
 	FNodalOffsetPointsCalibratorPointData LastCalibratorPoint;
 	LastCalibratorPoint.bIsValid = false;
@@ -724,7 +731,7 @@ bool UCameraNodalOffsetAlgoPoints::OnViewportClicked(const FGeometry& MyGeometry
 	}
 
 	// Check that we have a valid calibrator 3dpoint or camera data
-	if (!LastCalibratorPoint.bIsValid || !LastCameraData.bIsValid)
+	if (!LastCalibratorPoint.bIsValid)
 	{
 		return true;
 	}
@@ -925,8 +932,6 @@ bool UCameraNodalOffsetAlgoPoints::ValidateNewRow(TSharedPtr<FNodalOffsetPointsR
 		return false;
 	}
 
-	// FZ inputs are always valid, no need to verify them. They could be coming from LiveLink or fallback to a default one
-
 	// bApplyNodalOffset did not change.
 	//
 	// It can't change because we need to know if the camera pose is being affected or not by the current nodal offset evaluation.
diff --git a/Engine/Plugins/VirtualProduction/CameraCalibrationCore/Source/CameraCalibrationCore/Private/LensComponent.cpp b/Engine/Plugins/VirtualProduction/CameraCalibrationCore/Source/CameraCalibrationCore/Private/LensComponent.cpp
--- a/Engine/Plugins/VirtualProduction/CameraCalibrationCore/Source/CameraCalibrationCore/Private/LensComponent.cpp
+++ b/Engine/Plugins/VirtualProduction/CameraCalibrationCore/Source/CameraCalibrationCore/Private/LensComponent.cpp
@@ -780,8 +780,10 @@ void ULensComponent::UpdateLiveLinkFIZ(const ULiveLinkComponentController* const
 
 	if (StaticData->bIsFocusDistanceSupported && StaticData->bIsFocalLengthSupported)
 	{
-		LiveLinkFIZ.Input.Focus = FrameData->FocusDistance;
-		LiveLinkFIZ.Input.Zoom = FrameData->FocalLength;
+		LiveLinkFocus = FrameData->FocusDistance;
+		LiveLinkZoom = FrameData->FocalLength;
+
+		LiveLinkIris = FrameData->Aperture;
 
 		bWasLiveLinkFIZUpdated = true;
 	}
@@ -796,14 +798,16 @@ void ULensComponent::UpdateLensFileEvaluationInputs(UCineCameraComponent* CineCa
 	{
 		if (bWasLiveLinkFIZUpdated)
 		{
-			EvalInputs.Focus = LiveLinkFIZ.Input.Focus;
-			EvalInputs.Zoom = LiveLinkFIZ.Input.Zoom;
+			EvalInputs.Focus = LiveLinkFocus;
+			EvalInputs.Iris = LiveLinkIris;
+			EvalInputs.Zoom = LiveLinkZoom;
 			EvalInputs.bIsValid = true;
 		}
 	}
 	else if (EvaluationMode == EFIZEvaluationMode::UseCameraSettings)
 	{
 		EvalInputs.Focus = CineCameraComponent->CurrentFocusDistance;
+		EvalInputs.Iris = CineCameraComponent->CurrentAperture;
 		// The camera's focal length might have been altered by applying an overscan FOV. However, the LensFile needs to be evaluated with the non-overscanned value
 		EvalInputs.Zoom = OriginalFocalLength;
 		EvalInputs.bIsValid = true;
diff --git a/Engine/Plugins/VirtualProduction/CameraCalibrationCore/Source/CameraCalibrationCore/Public/CameraCalibrationSettings.h b/Engine/Plugins/VirtualProduction/CameraCalibrationCore/Source/CameraCalibrationCore/Public/CameraCalibrationSettings.h
--- a/Engine/Plugins/VirtualProduction/CameraCalibrationCore/Source/CameraCalibrationCore/Public/CameraCalibrationSettings.h
+++ b/Engine/Plugins/VirtualProduction/CameraCalibrationCore/Source/CameraCalibrationCore/Public/CameraCalibrationSettings.h
@@ -230,7 +230,7 @@ public:
 	/**
 	 * Time slider enable flag
 	 */
-	UPROPERTY(config, EditAnywhere, Category = "Settings", Meta = (ToolTip = "Enable or Disable Time input driven by Live Link."))
+	UPROPERTY(config, EditAnywhere, Category = "Settings", Meta = (ToolTip = "Enable or Disable Time input driven by evaluation inputs."))
 	bool bEnableTimeSlider = true;
 
 	/** 
diff --git a/Engine/Plugins/VirtualProduction/CameraCalibrationCore/Source/CameraCalibrationCore/Public/LensComponent.h b/Engine/Plugins/VirtualProduction/CameraCalibrationCore/Source/CameraCalibrationCore/Public/LensComponent.h
--- a/Engine/Plugins/VirtualProduction/CameraCalibrationCore/Source/CameraCalibrationCore/Public/LensComponent.h
+++ b/Engine/Plugins/VirtualProduction/CameraCalibrationCore/Source/CameraCalibrationCore/Public/LensComponent.h
@@ -50,28 +50,6 @@ enum class EDistortionSource : uint8
 	Manual,
 };
 
-USTRUCT(BlueprintType)
-struct FLensFileEvaluationInputs
-{
-	GENERATED_BODY()
-
-public:
-	UPROPERTY(Interp, VisibleAnywhere, BlueprintReadOnly, Category = "Lens File")
-	float Focus = 0.0f;
-
-	UPROPERTY(Interp, VisibleAnywhere, BlueprintReadOnly, Category = "Lens File")
-	float Iris = 0.0f;
-
-	UPROPERTY(Interp, VisibleAnywhere, BlueprintReadOnly, Category = "Lens File")
-	float Zoom = 0.0f;
-
-	UPROPERTY(Interp, VisibleAnywhere, BlueprintReadOnly, Category = "Lens File")
-	FCameraFilmbackSettings Filmback;
-
-	UPROPERTY(Transient, BlueprintReadOnly, Category = "Lens File")
-	bool bIsValid = false;
-};
-
 /** Component for applying a post-process lens distortion effect to a CineCameraComponent on the same actor */
 UCLASS(HideCategories=(Tags, Activation, Cooking, AssetUserData, Collision), meta=(BlueprintSpawnableComponent))
 class CAMERACALIBRATIONCORE_API ULensComponent : public UActorComponent
@@ -378,7 +356,9 @@ PRAGMA_ENABLE_DEPRECATION_WARNINGS
 
 private:
 	/** Latest LiveLink FIZ data, used to evaluate the LensFile */
-	FLensFileEvalData LiveLinkFIZ;
+	float LiveLinkFocus = 0.0f;
+	float LiveLinkIris = 0.0f;
+	float LiveLinkZoom = 0.0f;
 
 	/** Whether LiveLink FIZ was received this tick */
 	bool bWasLiveLinkFIZUpdated = false;
diff --git a/Engine/Plugins/VirtualProduction/CameraCalibrationCore/Source/CameraCalibrationCore/Public/LensFile.h b/Engine/Plugins/VirtualProduction/CameraCalibrationCore/Source/CameraCalibrationCore/Public/LensFile.h
--- a/Engine/Plugins/VirtualProduction/CameraCalibrationCore/Source/CameraCalibrationCore/Public/LensFile.h
+++ b/Engine/Plugins/VirtualProduction/CameraCalibrationCore/Source/CameraCalibrationCore/Public/LensFile.h
@@ -325,6 +325,28 @@ public:
 	TObjectPtr<ULensFile> LensFile = nullptr;
 };
 
+USTRUCT(BlueprintType)
+struct FLensFileEvaluationInputs
+{
+	GENERATED_BODY()
+
+public:
+	UPROPERTY(Interp, VisibleAnywhere, BlueprintReadOnly, Category = "Lens File")
+	float Focus = 0.0f;
+
+	UPROPERTY(Interp, VisibleAnywhere, BlueprintReadOnly, Category = "Lens File")
+	float Iris = 0.0f;
+
+	UPROPERTY(Interp, VisibleAnywhere, BlueprintReadOnly, Category = "Lens File")
+	float Zoom = 0.0f;
+
+	UPROPERTY(Interp, VisibleAnywhere, BlueprintReadOnly, Category = "Lens File")
+	FCameraFilmbackSettings Filmback;
+
+	UPROPERTY(Transient, BlueprintReadOnly, Category = "Lens File")
+	bool bIsValid = false;
+};
+
 /** Structure that caches the inputs (and other useful bits) used when evaluating the Lens File */
 struct FLensFileEvalData
 {
