[CameraCalibration] Fix for optical axis calibration pointing in the wrong (flipped) direction

diff --git a/Engine/Plugins/VirtualProduction/CameraCalibration/Source/CameraCalibrationEditor/Private/Calibrators/CameraNodalOffsetAlgoOpticalAxis.cpp b/Engine/Plugins/VirtualProduction/CameraCalibration/Source/CameraCalibrationEditor/Private/Calibrators/CameraNodalOffsetAlgoOpticalAxis.cpp
--- a/Engine/Plugins/VirtualProduction/CameraCalibration/Source/CameraCalibrationEditor/Private/Calibrators/CameraNodalOffsetAlgoOpticalAxis.cpp
+++ b/Engine/Plugins/VirtualProduction/CameraCalibration/Source/CameraCalibrationEditor/Private/Calibrators/CameraNodalOffsetAlgoOpticalAxis.cpp
@@ -252,7 +252,7 @@ bool UCameraNodalOffsetAlgoOpticalAxis::GetNodalOffset(FNodalPointOffset& OutNod
 
 	// Find a look at rotation for the camera
 	const FVector Point0 = PointAlongAxis + (0.0f * OpticalAxis.GetValue());
-	const FVector Point1 = PointAlongAxis + (1.0f * OpticalAxis.GetValue());
+	const FVector Point1 = PointAlongAxis - (1.0f * OpticalAxis.GetValue());
 
 	LookAtRotation = UKismetMathLibrary::FindLookAtRotation(Point1, Point0);
 
