CameraCalibration - Fix simulcam viewport to match resolution of selected media source

robomerge[FN_EngineMerge] -UE5-RES

diff --git a/Engine/Plugins/VirtualProduction/CameraCalibration/Source/CameraCalibrationEditor/Private/AssetEditor/CameraCalibrationStepsController.cpp b/Engine/Plugins/VirtualProduction/CameraCalibration/Source/CameraCalibrationEditor/Private/AssetEditor/CameraCalibrationStepsController.cpp
--- a/Engine/Plugins/VirtualProduction/CameraCalibration/Source/CameraCalibrationEditor/Private/AssetEditor/CameraCalibrationStepsController.cpp
+++ b/Engine/Plugins/VirtualProduction/CameraCalibration/Source/CameraCalibrationEditor/Private/AssetEditor/CameraCalibrationStepsController.cpp
@@ -183,6 +183,27 @@ bool FCameraCalibrationStepsController::OnTick(float DeltaTime)
 		LensFileEvaluationInputs = LensComponent->GetLensFileEvaluationInputs();
 	}
 
+	// Compare the current dimensions of the playing media track to the comp's render resolution to determine if the comp needs to be resized
+	if (MediaPlayer.IsValid())
+	{
+		const FIntPoint MediaDimensions = MediaPlayer->GetVideoTrackDimensions(INDEX_NONE, INDEX_NONE);
+
+		// If no track was found, the dimensions might be (0, 0)
+		if (MediaDimensions.X != 0 && MediaDimensions.Y != 0)
+		{
+			if (RenderTarget->SizeX != MediaDimensions.X || RenderTarget->SizeY != MediaDimensions.Y)
+			{
+				// Resize the media plate comp layer and its output render target to match the incoming media dimensions
+				MediaPlateRenderTarget->ResizeTarget(MediaDimensions.X, MediaDimensions.Y);
+				MediaPlate->SetRenderResolution(MediaDimensions);
+
+				// Resize the parent comp and its output render target to match the incoming media dimensions
+				RenderTarget->ResizeTarget(MediaDimensions.X, MediaDimensions.Y);
+				Comp->SetRenderResolution(MediaDimensions);
+			}
+		}
+	}
+
 	for (TStrongObjectPtr<UCameraCalibrationStep>& Step : CalibrationSteps)
 	{
 		if (Step.IsValid())
diff --git a/Engine/Plugins/VirtualProduction/CameraCalibration/Source/CameraCalibrationEditor/Private/AssetEditor/Viewport/SSimulcamEditorViewport.cpp b/Engine/Plugins/VirtualProduction/CameraCalibration/Source/CameraCalibrationEditor/Private/AssetEditor/Viewport/SSimulcamEditorViewport.cpp
--- a/Engine/Plugins/VirtualProduction/CameraCalibration/Source/CameraCalibrationEditor/Private/AssetEditor/Viewport/SSimulcamEditorViewport.cpp
+++ b/Engine/Plugins/VirtualProduction/CameraCalibration/Source/CameraCalibrationEditor/Private/AssetEditor/Viewport/SSimulcamEditorViewport.cpp
@@ -149,7 +149,12 @@ void SSimulcamEditorViewport::CacheEffectiveTextureSize()
 		return;
 	}
 
-	CachedEffectiveTextureSize = FVector2D(Texture->GetSurfaceWidth(), Texture->GetSurfaceHeight());
+	// If the cached texture dimensions do not match the current texture dimensions (because the texture was resized), update the cached value and zoom all the way out
+	if (CachedEffectiveTextureSize.X != Texture->GetSurfaceWidth() || CachedEffectiveTextureSize.Y != Texture->GetSurfaceHeight())
+	{
+		CachedEffectiveTextureSize = FVector2D(Texture->GetSurfaceWidth(), Texture->GetSurfaceHeight());
+		ViewportClient->OnTextureResized();
+	}
 }
 
 FVector2D SSimulcamEditorViewport::CalculateTextureDimensions() const
diff --git a/Engine/Plugins/VirtualProduction/CameraCalibration/Source/CameraCalibrationEditor/Private/AssetEditor/Viewport/SimulcamEditorViewportClient.cpp b/Engine/Plugins/VirtualProduction/CameraCalibration/Source/CameraCalibrationEditor/Private/AssetEditor/Viewport/SimulcamEditorViewportClient.cpp
--- a/Engine/Plugins/VirtualProduction/CameraCalibration/Source/CameraCalibrationEditor/Private/AssetEditor/Viewport/SimulcamEditorViewportClient.cpp
+++ b/Engine/Plugins/VirtualProduction/CameraCalibration/Source/CameraCalibrationEditor/Private/AssetEditor/Viewport/SimulcamEditorViewportClient.cpp
@@ -222,6 +222,17 @@ void FSimulcamEditorViewportClient::OnViewportResized(FViewport* InViewport, uin
 	}
 }
 
+void FSimulcamEditorViewportClient::OnTextureResized()
+{
+	if (TSharedPtr<SSimulcamEditorViewport> SimulcamEditorViewport = SimulcamEditorViewportWeakPtr.Pin())
+	{
+		if (FViewport* Viewport = SimulcamEditorViewport->GetViewport()->GetViewport())
+		{
+			ZoomToFit(Viewport);
+		}
+	}
+}
+
 void FSimulcamEditorViewportClient::ZoomToFit(FViewport* InViewport)
 {
 	if (SimulcamEditorViewportWeakPtr.IsValid())
diff --git a/Engine/Plugins/VirtualProduction/CameraCalibration/Source/CameraCalibrationEditor/Private/AssetEditor/Viewport/SimulcamEditorViewportClient.h b/Engine/Plugins/VirtualProduction/CameraCalibration/Source/CameraCalibrationEditor/Private/AssetEditor/Viewport/SimulcamEditorViewportClient.h
--- a/Engine/Plugins/VirtualProduction/CameraCalibration/Source/CameraCalibrationEditor/Private/AssetEditor/Viewport/SimulcamEditorViewportClient.h
+++ b/Engine/Plugins/VirtualProduction/CameraCalibration/Source/CameraCalibrationEditor/Private/AssetEditor/Viewport/SimulcamEditorViewportClient.h
@@ -36,6 +36,10 @@ public:
 
 	/** Triggered whenever the viewport is resized */
 	void OnViewportResized(FViewport* InViewport, uint32 InParams);
+
+	/** Triggered whenever the underlying texture resource is resized */
+	void OnTextureResized();
+
 private:
 	FVector2D GetTexturePosition() const;
 	
