CameraCalibration - Fix crash in steps controller when media player is not valid

diff --git a/Engine/Plugins/VirtualProduction/CameraCalibration/Source/CameraCalibrationEditor/Private/AssetEditor/CameraCalibrationStepsController.cpp b/Engine/Plugins/VirtualProduction/CameraCalibration/Source/CameraCalibrationEditor/Private/AssetEditor/CameraCalibrationStepsController.cpp
--- a/Engine/Plugins/VirtualProduction/CameraCalibration/Source/CameraCalibrationEditor/Private/AssetEditor/CameraCalibrationStepsController.cpp
+++ b/Engine/Plugins/VirtualProduction/CameraCalibration/Source/CameraCalibrationEditor/Private/AssetEditor/CameraCalibrationStepsController.cpp
@@ -893,92 +893,114 @@ bool FCameraCalibrationStepsController::OnSimulcamViewportInputKey(const FKey& I
 FReply FCameraCalibrationStepsController::OnRewindButtonClicked()
 {
 	// Rewind to the beginning of the media
-	MediaPlayer->Rewind();
+	if (MediaPlayer.IsValid())
+	{
+		MediaPlayer->Rewind();
+	}
+
 	return FReply::Handled();
 }
 
 FReply FCameraCalibrationStepsController::OnReverseButtonClicked()
 {
 	// Increase the reverse media playback rate
-	MediaPlayer->SetRate(GetFasterReverseRate());
+	if (MediaPlayer.IsValid())
+	{
+		MediaPlayer->SetRate(GetFasterReverseRate());
+	}
+
 	return FReply::Handled();
 }
 
 FReply FCameraCalibrationStepsController::OnStepBackButtonClicked() 
 {
-	const float DefaultStepRateInMilliseconds = GetDefault<UCameraCalibrationEditorSettings>()->DefaultMediaStepRateInMilliseconds;
-	const bool bForceDefaultStepRate = GetDefault<UCameraCalibrationEditorSettings>()->bForceDefaultMediaStepRate;
+	if (MediaPlayer.IsValid())
+	{
+		const float DefaultStepRateInMilliseconds = GetDefault<UCameraCalibrationEditorSettings>()->DefaultMediaStepRateInMilliseconds;
+		const bool bForceDefaultStepRate = GetDefault<UCameraCalibrationEditorSettings>()->bForceDefaultMediaStepRate;
 
-	// The media player could return a frame rate of 0 for the current video track
-	const float MediaFrameRate = MediaPlayer->GetVideoTrackFrameRate(INDEX_NONE, INDEX_NONE);
+		// The media player could return a frame rate of 0 for the current video track
+		const float MediaFrameRate = MediaPlayer->GetVideoTrackFrameRate(INDEX_NONE, INDEX_NONE);
 
-	// Use the default step rate if the media player returned an invalid frame rate or if the project settings force it
-	float MillisecondsPerStep = 0.0f;
-	if (FMath::IsNearlyEqual(MediaFrameRate, 0.0f) || bForceDefaultStepRate)
-	{
-		MillisecondsPerStep = DefaultStepRateInMilliseconds;
-	}
-	else
-	{
-		MillisecondsPerStep = 1000.0f / MediaFrameRate;
-	}
+		// Use the default step rate if the media player returned an invalid frame rate or if the project settings force it
+		float MillisecondsPerStep = 0.0f;
+		if (FMath::IsNearlyEqual(MediaFrameRate, 0.0f) || bForceDefaultStepRate)
+		{
+			MillisecondsPerStep = DefaultStepRateInMilliseconds;
+		}
+		else
+		{
+			MillisecondsPerStep = 1000.0f / MediaFrameRate;
+		}
 
-	// Compute the number of ticks in one step and go backward from the media's current time (clamping to 0)
-	const FTimespan TicksInOneStep = ETimespan::TicksPerMillisecond * (MillisecondsPerStep);
-	FTimespan PreviousStepTime = MediaPlayer->GetTime() - TicksInOneStep;
-	if (PreviousStepTime < FTimespan::Zero())
-	{
-		PreviousStepTime = FTimespan::Zero();
-	}
+		// Compute the number of ticks in one step and go backward from the media's current time (clamping to 0)
+		const FTimespan TicksInOneStep = ETimespan::TicksPerMillisecond * (MillisecondsPerStep);
+		FTimespan PreviousStepTime = MediaPlayer->GetTime() - TicksInOneStep;
+		if (PreviousStepTime < FTimespan::Zero())
+		{
+			PreviousStepTime = FTimespan::Zero();
+		}
 
-	MediaPlayer->Seek(PreviousStepTime);
-	MediaPlayer->Pause();
+		MediaPlayer->Seek(PreviousStepTime);
+		MediaPlayer->Pause();
+	}
 
 	return FReply::Handled();
 }
 
 FReply FCameraCalibrationStepsController::OnPlayButtonClicked() 
 {
-	MediaPlayer->Play();
+	if (MediaPlayer.IsValid())
+	{
+		MediaPlayer->Play();
+	}
+
 	return FReply::Handled(); 
 }
 
 FReply FCameraCalibrationStepsController::OnPauseButtonClicked() 
 {
-	MediaPlayer->Pause();
+	if (MediaPlayer.IsValid())
+	{
+		MediaPlayer->Pause();
+	}
+
 	return FReply::Handled(); 
 }
 
 FReply FCameraCalibrationStepsController::OnStepForwardButtonClicked()
 {
-	const float DefaultStepRateInMilliseconds = GetDefault<UCameraCalibrationEditorSettings>()->DefaultMediaStepRateInMilliseconds;
-	const bool bForceDefaultStepRate = GetDefault<UCameraCalibrationEditorSettings>()->bForceDefaultMediaStepRate;
+	if (MediaPlayer.IsValid())
+	{
+		const float DefaultStepRateInMilliseconds = GetDefault<UCameraCalibrationEditorSettings>()->DefaultMediaStepRateInMilliseconds;
+		const bool bForceDefaultStepRate = GetDefault<UCameraCalibrationEditorSettings>()->bForceDefaultMediaStepRate;
 
-	// The media player could return a frame rate of 0 for the current video track
-	const float MediaFrameRate = MediaPlayer->GetVideoTrackFrameRate(INDEX_NONE, INDEX_NONE);
+		// The media player could return a frame rate of 0 for the current video track
+		const float MediaFrameRate = MediaPlayer->GetVideoTrackFrameRate(INDEX_NONE, INDEX_NONE);
 
-	// Use the default step rate if the media player returned an invalid frame rate or if the project settings force it
-	float MillisecondsPerStep = 0.0f;
-	if (FMath::IsNearlyEqual(MediaFrameRate, 0.0f) || bForceDefaultStepRate)
-	{
-		MillisecondsPerStep = DefaultStepRateInMilliseconds;
-	}
-	else
-	{
-		MillisecondsPerStep = 1000.0f / MediaFrameRate;
-	}
+		// Use the default step rate if the media player returned an invalid frame rate or if the project settings force it
+		float MillisecondsPerStep = 0.0f;
+		if (FMath::IsNearlyEqual(MediaFrameRate, 0.0f) || bForceDefaultStepRate)
+		{
+			MillisecondsPerStep = DefaultStepRateInMilliseconds;
+		}
+		else
+		{
+			MillisecondsPerStep = 1000.0f / MediaFrameRate;
+		}
 
 
-	// Compute the number of ticks in one step and go forward from the media's current time
-	const FTimespan TicksInOneStep = ETimespan::TicksPerMillisecond * (MillisecondsPerStep);
-	const FTimespan NextStepTime = MediaPlayer->GetTime() + TicksInOneStep;
+		// Compute the number of ticks in one step and go forward from the media's current time
+		const FTimespan TicksInOneStep = ETimespan::TicksPerMillisecond * (MillisecondsPerStep);
+		const FTimespan NextStepTime = MediaPlayer->GetTime() + TicksInOneStep;
 
-	// Ensure that we do not attempt to seek past the end of the media
-	const FTimespan Duration = MediaPlayer->GetDuration();
-	if ((NextStepTime + TicksInOneStep) < Duration)
-	{
-		MediaPlayer->Seek(NextStepTime);
-		MediaPlayer->Pause();
+		// Ensure that we do not attempt to seek past the end of the media
+		const FTimespan Duration = MediaPlayer->GetDuration();
+		if ((NextStepTime + TicksInOneStep) < Duration)
+		{
+			MediaPlayer->Seek(NextStepTime);
+			MediaPlayer->Pause();
+		}
 	}
 
 	return FReply::Handled();
@@ -987,55 +1009,84 @@ FReply FCameraCalibrationStepsController::OnStepForwardButtonClicked()
 FReply FCameraCalibrationStepsController::OnForwardButtonClicked() 
 {
 	// Increase the forward media playback rate
-	MediaPlayer->SetRate(GetFasterForwardRate());
+	if (MediaPlayer.IsValid())
+	{
+		MediaPlayer->SetRate(GetFasterForwardRate());
+	}
+
 	return FReply::Handled(); 
 }
 
 bool FCameraCalibrationStepsController::DoesMediaSupportSeeking() const
 {
-	return MediaPlayer->SupportsSeeking();
+	if (MediaPlayer.IsValid())
+	{
+		return MediaPlayer->SupportsSeeking();
+	}
+
+	return false;
 }
 
 bool FCameraCalibrationStepsController::DoesMediaSupportNextReverseRate() const
 {
-	constexpr bool Unthinned = false;
-	return MediaPlayer->SupportsRate(GetFasterReverseRate(), Unthinned);
+	if (MediaPlayer.IsValid())
+	{
+		constexpr bool Unthinned = false;
+		return MediaPlayer->SupportsRate(GetFasterReverseRate(), Unthinned);
+	}
+
+	return false;
 }
 
 bool FCameraCalibrationStepsController::DoesMediaSupportNextForwardRate() const
 {
-	constexpr bool Unthinned = false;
-	return MediaPlayer->SupportsRate(GetFasterForwardRate(), Unthinned);
+	if (MediaPlayer.IsValid())
+	{
+		constexpr bool Unthinned = false;
+		return MediaPlayer->SupportsRate(GetFasterForwardRate(), Unthinned);
+	}
+
+	return false;
 }
 
 float FCameraCalibrationStepsController::GetFasterReverseRate() const
 {
-	// Get the current playback rate of the media player
-	float Rate = MediaPlayer->GetRate();
-
-	// Reverse the playback direction (if needed) to ensure the rate is going in reverse
-	if (Rate > -1.0f)
+	if (MediaPlayer.IsValid())
 	{
-		return -1.0f;
+		// Get the current playback rate of the media player
+		float Rate = MediaPlayer->GetRate();
+
+		// Reverse the playback direction (if needed) to ensure the rate is going in reverse
+		if (Rate > -1.0f)
+		{
+			return -1.0f;
+		}
+
+		// Double the reverse playback rate
+		return 2.0f * Rate;
 	}
 
-	// Double the reverse playback rate
-	return 2.0f * Rate;
+	return -1.0f;
 }
 
 float FCameraCalibrationStepsController::GetFasterForwardRate() const
 {
-	// Get the current playback rate of the media player
-	float Rate = MediaPlayer->GetRate();
-
-	// Reverse the playback direction (if needed) to ensure the rate is going forward
-	if (Rate < 1.0f)
+	if (MediaPlayer.IsValid())
 	{
-		Rate = 1.0f;
+		// Get the current playback rate of the media player
+		float Rate = MediaPlayer->GetRate();
+
+		// Reverse the playback direction (if needed) to ensure the rate is going forward
+		if (Rate < 1.0f)
+		{
+			Rate = 1.0f;
+		}
+
+		// Double the forward playback rate
+		return 2.0f * Rate;
 	}
 
-	// Double the forward playback rate
-	return 2.0f * Rate;
+	return 1.0f;
 }
 
 void FCameraCalibrationStepsController::ToggleShowMediaPlaybackControls()
