CameraCalibration - Fix the simulcam comp material to not stretch CG when compositing if the CG and media layers have different aspect ratios

diff --git a/Engine/Plugins/VirtualProduction/CameraCalibration/Source/CameraCalibrationEditor/Private/AssetEditor/CameraCalibrationStepsController.cpp b/Engine/Plugins/VirtualProduction/CameraCalibration/Source/CameraCalibrationEditor/Private/AssetEditor/CameraCalibrationStepsController.cpp
--- a/Engine/Plugins/VirtualProduction/CameraCalibration/Source/CameraCalibrationEditor/Private/AssetEditor/CameraCalibrationStepsController.cpp
+++ b/Engine/Plugins/VirtualProduction/CameraCalibration/Source/CameraCalibrationEditor/Private/AssetEditor/CameraCalibrationStepsController.cpp
@@ -183,6 +183,8 @@ bool FCameraCalibrationStepsController::OnTick(float DeltaTime)
 		LensFileEvaluationInputs = LensComponent->GetLensFileEvaluationInputs();
 	}
 
+	bool bCompElementsResized = false;
+
 	// Compare the current dimensions of the playing media track to the comp's render resolution to determine if the comp needs to be resized
 	if (MediaPlayer.IsValid())
 	{
@@ -200,6 +202,69 @@ bool FCameraCalibrationStepsController::OnTick(float DeltaTime)
 				// Resize the parent comp and its output render target to match the incoming media dimensions
 				RenderTarget->ResizeTarget(MediaDimensions.X, MediaDimensions.Y);
 				Comp->SetRenderResolution(MediaDimensions);
+
+				bCompElementsResized = true;
+			}
+		}
+
+		if (ULensFile* LensFilePtr = GetLensFile())
+		{
+			LensFilePtr->SimulcamInfo.MediaResolution = MediaDimensions;
+			LensFilePtr->SimulcamInfo.MediaPlateAspectRatio = (MediaDimensions.Y != 0) ? (float)MediaDimensions.X / (float)MediaDimensions.Y : 0.0f;
+		}
+	}
+
+	// Compare the aspect ratio of the selected camera to the cg layer's aspect ratio to determine if the cg layer needs to be resized
+	if (CineCameraComponent.Get())
+	{
+		const float CameraAspectRatio = CineCameraComponent->Filmback.SensorAspectRatio;
+
+		FIntPoint CGResolution = CGLayer->GetRenderResolution();
+		const float CGLayerAspectRatio = (CGResolution.Y != 0) ? (float)CGResolution.X / (float)CGResolution.Y : 0.0f;
+
+		constexpr float Tolerance = 0.01;
+		if (!FMath::IsNearlyEqual(CameraAspectRatio, CGLayerAspectRatio, Tolerance))
+		{
+			CGResolution.X = CGResolution.Y * CameraAspectRatio;
+
+			CGLayer->ResolutionSource = EInheritedSourceType::Override;
+			CGLayer->SetRenderResolution(CGResolution);
+
+			bCompElementsResized = true;
+		}
+
+		if (ULensFile* LensFilePtr = GetLensFile())
+		{
+			LensFilePtr->SimulcamInfo.CGLayerAspectRatio = CameraAspectRatio;
+		}
+	}
+
+	// If either the CG or the Media was resized, then recompute the aspect ratio correction for the comp material
+	if (bCompElementsResized)
+	{
+		if (MaterialPass.IsValid())
+		{
+			const float MediaAspectRatio = (MediaPlateRenderTarget->SizeY != 0) ? (float)MediaPlateRenderTarget->SizeX / (float)MediaPlateRenderTarget->SizeY : 0.0f;
+
+			const FIntPoint CGResolution = CGLayer->GetRenderResolution();
+			const float CGLayerAspectRatio = (CGResolution.Y != 0) ? (float)CGResolution.X / (float)CGResolution.Y : 0.0f;
+
+			if (CGLayerAspectRatio > MediaAspectRatio)
+			{
+				const float AspectRatioCorrection = (MediaAspectRatio != 0.0f) ? CGLayerAspectRatio / MediaAspectRatio : 1.0f;
+				MaterialPass->Material.SetScalarOverride(TEXT("AspectRatioCorrection_H"), 1.0f);
+				MaterialPass->Material.SetScalarOverride(TEXT("AspectRatioCorrection_V"), AspectRatioCorrection);
+			}
+			else if (MediaAspectRatio > CGLayerAspectRatio)
+			{
+				const float AspectRatioCorrection = (CGLayerAspectRatio != 0.0f) ? MediaAspectRatio / CGLayerAspectRatio : 1.0f;
+				MaterialPass->Material.SetScalarOverride(TEXT("AspectRatioCorrection_H"), AspectRatioCorrection);
+				MaterialPass->Material.SetScalarOverride(TEXT("AspectRatioCorrection_V"), 1.0f);
+			}
+			else
+			{
+				MaterialPass->Material.SetScalarOverride(TEXT("AspectRatioCorrection_H"), 1.0f);
+				MaterialPass->Material.SetScalarOverride(TEXT("AspectRatioCorrection_V"), 1.0f);
 			}
 		}
 	}
@@ -714,7 +779,7 @@ void FCameraCalibrationStepsController::CreateMediaPlateOutput()
 	URenderTargetCompositingOutput* RTOutput = Cast<URenderTargetCompositingOutput>(MediaPlate->CreateNewOutputPass(
 		TEXT("MediaPlateOutput"),
 		URenderTargetCompositingOutput::StaticClass())
-	);
+		);
 
 	if (!RTOutput)
 	{
@@ -727,7 +792,7 @@ void FCameraCalibrationStepsController::CreateMediaPlateOutput()
 	MediaPlateRenderTarget = NewObject<UTextureRenderTarget2D>(
 		GetTransientPackage(),
 		MakeUniqueObjectName(GetTransientPackage(), UTextureRenderTarget2D::StaticClass())
-	);
+		);
 
 	if (!MediaPlateRenderTarget.IsValid())
 	{
@@ -735,7 +800,7 @@ void FCameraCalibrationStepsController::CreateMediaPlateOutput()
 		Cleanup();
 		return;
 	}
-	
+
 	MediaPlateRenderTarget->RenderTargetFormat = ETextureRenderTargetFormat::RTF_RGBA8;
 	MediaPlateRenderTarget->ClearColor = FLinearColor::Black;
 	MediaPlateRenderTarget->bAutoGenerateMips = false;
@@ -813,6 +878,14 @@ void FCameraCalibrationStepsController::SetCamera(ACameraActor* InCamera)
 {
 	Camera = InCamera;
 
+	TInlineComponentArray<UCineCameraComponent*> CameraComponents;
+	InCamera->GetComponents(CameraComponents);
+
+	if (CameraComponents.Num() > 0)
+	{
+		CineCameraComponent = CameraComponents[0];
+	}
+
 	// Update the Comp with this camera
 	if (Comp.IsValid())
 	{
diff --git a/Engine/Plugins/VirtualProduction/CameraCalibration/Source/CameraCalibrationEditor/Private/AssetEditor/CameraCalibrationStepsController.h b/Engine/Plugins/VirtualProduction/CameraCalibration/Source/CameraCalibrationEditor/Private/AssetEditor/CameraCalibrationStepsController.h
--- a/Engine/Plugins/VirtualProduction/CameraCalibration/Source/CameraCalibrationEditor/Private/AssetEditor/CameraCalibrationStepsController.h
+++ b/Engine/Plugins/VirtualProduction/CameraCalibration/Source/CameraCalibrationEditor/Private/AssetEditor/CameraCalibrationStepsController.h
@@ -287,6 +287,9 @@ private:
 	/** The currently selected camera */
 	TWeakObjectPtr<ACameraActor> Camera;
 
+	/** The currently selected camera's camera component */
+	TWeakObjectPtr<UCineCameraComponent> CineCameraComponent;
+
 	/** The delegate for the core ticker callback */
 	FTSTicker::FDelegateHandle TickerHandle;
 
diff --git a/Engine/Plugins/VirtualProduction/CameraCalibration/Source/CameraCalibrationEditor/Private/AssetEditor/LensInfoStep.cpp b/Engine/Plugins/VirtualProduction/CameraCalibration/Source/CameraCalibrationEditor/Private/AssetEditor/LensInfoStep.cpp
--- a/Engine/Plugins/VirtualProduction/CameraCalibration/Source/CameraCalibrationEditor/Private/AssetEditor/LensInfoStep.cpp
+++ b/Engine/Plugins/VirtualProduction/CameraCalibration/Source/CameraCalibrationEditor/Private/AssetEditor/LensInfoStep.cpp
@@ -25,17 +25,17 @@ void ULensInfoStep::Initialize(TWeakPtr<FCameraCalibrationStepsController> InCam
 
 TSharedRef<SWidget> ULensInfoStep::BuildUI()
 {
+	FPropertyEditorModule& PropertyEditor = FModuleManager::Get().LoadModuleChecked<FPropertyEditorModule>(TEXT("PropertyEditor"));
+
+	ULensFile* LensFile = CameraCalibrationStepsController.Pin()->GetLensFile();
+
 	FStructureDetailsViewArgs LensInfoStructDetailsView;
 	FDetailsViewArgs DetailArgs;
 	DetailArgs.bAllowSearch = false;
 	DetailArgs.bShowScrollBar = true;
 
-	FPropertyEditorModule& PropertyEditor = FModuleManager::Get().LoadModuleChecked<FPropertyEditorModule>(TEXT("PropertyEditor"));
-
-	ULensFile* LensFile = CameraCalibrationStepsController.Pin()->GetLensFile();
-	TSharedRef<FStructOnScope> StructOnScope = MakeShared<FStructOnScope>(FLensInfo::StaticStruct(), reinterpret_cast<uint8*>(&LensFile->LensInfo));
-
-	TSharedPtr<IStructureDetailsView> LensInfoStructureDetailsView = PropertyEditor.CreateStructureDetailView(DetailArgs, LensInfoStructDetailsView, StructOnScope);
+	TSharedRef<FStructOnScope> LensInfoStructOnScope = MakeShared<FStructOnScope>(FLensInfo::StaticStruct(), reinterpret_cast<uint8*>(&LensFile->LensInfo));
+	TSharedPtr<IStructureDetailsView> LensInfoStructureDetailsView = PropertyEditor.CreateStructureDetailView(DetailArgs, LensInfoStructDetailsView, LensInfoStructOnScope);
 
 	LensInfoStructureDetailsView->GetOnFinishedChangingPropertiesDelegate().AddLambda([&](const FPropertyChangedEvent& PropertyChangedEvent)
 	{
@@ -46,6 +46,13 @@ TSharedRef<SWidget> ULensInfoStep::BuildUI()
 		}
 	});
 
+	FStructureDetailsViewArgs SimulcamInfoStructDetailsView;
+
+	TSharedRef<FStructOnScope> SimulcamInfoStructOnScope = MakeShared<FStructOnScope>(FSimulcamInfo::StaticStruct(), reinterpret_cast<uint8*>(&LensFile->SimulcamInfo));
+	TSharedPtr<IStructureDetailsView> SimulcamInfoStructureDetailsView = PropertyEditor.CreateStructureDetailView(DetailArgs, SimulcamInfoStructDetailsView, SimulcamInfoStructOnScope);
+
+	SimulcamInfoStructureDetailsView->GetDetailsView()->SetIsPropertyEditingEnabledDelegate(FIsPropertyEditingEnabled::CreateLambda([]() { return false; }));
+
 	TSharedPtr<SWidget> StepWidget = 
 		SNew(SVerticalBox)
 		+ SVerticalBox::Slot() // Lens information structure
@@ -64,6 +71,14 @@ TSharedRef<SWidget> ULensInfoStep::BuildUI()
 				.OnResetToDefault(FSimpleDelegate::CreateUObject(this, &ULensInfoStep::ResetToDefault))
 				.DiffersFromDefault(TAttribute<bool>::Create(TAttribute<bool>::FGetter::CreateUObject(this, &ULensInfoStep::DiffersFromDefault)))
 			]
+		]
+	
+		+ SVerticalBox::Slot() // Simulcam information structure
+		.AutoHeight()
+		[
+			SNew(SHorizontalBox)
+			+ SHorizontalBox::Slot()
+			[ SimulcamInfoStructureDetailsView->GetWidget().ToSharedRef() ]
 		];
 	
 
diff --git a/Engine/Plugins/VirtualProduction/CameraCalibrationCore/Source/CameraCalibrationCore/Public/LensData.h b/Engine/Plugins/VirtualProduction/CameraCalibrationCore/Source/CameraCalibrationCore/Public/LensData.h
--- a/Engine/Plugins/VirtualProduction/CameraCalibrationCore/Source/CameraCalibrationCore/Public/LensData.h
+++ b/Engine/Plugins/VirtualProduction/CameraCalibrationCore/Source/CameraCalibrationCore/Public/LensData.h
@@ -47,15 +47,41 @@ public:
 	TSubclassOf<ULensModel> LensModel;
 
 	/** Width and height of the calibrated camera's sensor, in millimeters */
-	UPROPERTY(EditAnywhere, BlueprintReadWrite, Category = "Lens Info")
+	UPROPERTY(EditAnywhere, BlueprintReadWrite, Category = "Camera Info")
 	FVector2D SensorDimensions = FVector2D(23.76f, 13.365f);
 
-	/** Camera image size, in pixels */
-	UPROPERTY(EditAnywhere, BlueprintReadWrite, Category = "Lens Info")
+	/** 
+	  * Resolution of the original footage that was captured by the camera (not necessarily the resolution of the media source).
+	  * For example, the original footage might have been 4320x1746, but to transmit that image over SDI, it might have been scaled and fit into a 4096x2160 frame.
+	  * In this case, the "Image Resolution" should be set to 4320x1746, while the "Media Resolution" will read 4096x2160.
+	  */
+	UPROPERTY(EditAnywhere, BlueprintReadWrite, Category = "Camera Info", meta = (DisplayName = "Image Resolution"))
 	FIntPoint ImageDimensions = FIntPoint(1920, 1080);
 
 };
 
+/**
+ * Information about the simulcam composition
+ */
+USTRUCT(BlueprintType)
+struct CAMERACALIBRATIONCORE_API FSimulcamInfo
+{
+	GENERATED_BODY()
+
+public:
+	/** Resolution of the selected Media Source */
+	UPROPERTY(EditAnywhere, BlueprintReadWrite, Category = "Media Info")
+	FIntPoint MediaResolution = FIntPoint(0, 0);
+
+	/** Aspect ratio of the media plate in the simulcam composition, which was set to match the aspect ratio of the selected Media Source */
+	UPROPERTY(EditAnywhere, BlueprintReadWrite, Category = "Simulcam Info")
+	float MediaPlateAspectRatio = 0.0f;
+
+	/** Aspect ratio of the CG layer in the simulcam composition, which was set to match the aspect ratio of the selected Camera Actor */
+	UPROPERTY(EditAnywhere, BlueprintReadWrite, Category = "Simulcam Info", meta = (DisplayName = "CG Layer Aspect Ratio"))
+	float CGLayerAspectRatio = 0.0f;
+};
+
 /**
  * Lens distortion parameters
  */
diff --git a/Engine/Plugins/VirtualProduction/CameraCalibrationCore/Source/CameraCalibrationCore/Public/LensFile.h b/Engine/Plugins/VirtualProduction/CameraCalibrationCore/Source/CameraCalibrationCore/Public/LensFile.h
--- a/Engine/Plugins/VirtualProduction/CameraCalibrationCore/Source/CameraCalibrationCore/Public/LensFile.h
+++ b/Engine/Plugins/VirtualProduction/CameraCalibrationCore/Source/CameraCalibrationCore/Public/LensFile.h
@@ -238,6 +238,10 @@ public:
 	UPROPERTY(EditAnywhere, BlueprintReadWrite, Category = "Lens info")
 	FLensInfo LensInfo;
 
+	/** Simulcam information */
+	UPROPERTY(VisibleAnywhere, BlueprintReadOnly, Category = "Simulcam info")
+	FSimulcamInfo SimulcamInfo;
+
 	/** Type of data used for lens mapping */
 	UPROPERTY(EditAnywhere, BlueprintReadWrite, Category = "Lens mapping")
 	ELensDataMode DataMode = ELensDataMode::Parameters;
