Add "camera feed" info setting to Lens File editor to allow users to specify a subregion of their media that contains the actual camera feed.

diff --git a/Engine/Plugins/VirtualProduction/CameraCalibration/Source/CameraCalibrationEditor/Private/AssetEditor/CameraCalibrationStepsController.cpp b/Engine/Plugins/VirtualProduction/CameraCalibration/Source/CameraCalibrationEditor/Private/AssetEditor/CameraCalibrationStepsController.cpp
--- a/Engine/Plugins/VirtualProduction/CameraCalibration/Source/CameraCalibrationEditor/Private/AssetEditor/CameraCalibrationStepsController.cpp
+++ b/Engine/Plugins/VirtualProduction/CameraCalibration/Source/CameraCalibrationEditor/Private/AssetEditor/CameraCalibrationStepsController.cpp
@@ -242,42 +242,141 @@ bool FCameraCalibrationStepsController::OnTick(float DeltaTime)
 	// If either the CG or the Media was resized, then recompute the aspect ratio correction for the comp material
 	if (bCompElementsResized)
 	{
-		if (MaterialPass.IsValid())
+		if (ULensFile* LensFilePtr = GetLensFile())
 		{
-			const float MediaAspectRatio = (MediaPlateRenderTarget->SizeY != 0) ? (float)MediaPlateRenderTarget->SizeX / (float)MediaPlateRenderTarget->SizeY : 0.0f;
+			UpdateAspectRatioCorrection(LensFilePtr->CameraFeedInfo.GetDimensions());
+		}
+	}
+
+	for (TStrongObjectPtr<UCameraCalibrationStep>& Step : CalibrationSteps)
+	{
+		if (Step.IsValid())
+		{
+			Step->Tick(DeltaTime);
+		}
+	}
+
+	return true;
+}
+
+void FCameraCalibrationStepsController::UpdateCameraFeedInfo(FVector2D MousePosition)
+{
+	// Compute the size of the camera feed, using the MousePosition as one of its corners.
+	// We assume that the camera feed will always be centered in the media. If that assumption proves false, this math should be updated in the future.
+	FIntPoint CameraFeedDimensions = FIntPoint(0, 0);
+	CameraFeedDimensions.X = FMath::Abs((MediaPlateRenderTarget->SizeX / 2) - FMath::Floor(MousePosition.X)) * 2;
+	CameraFeedDimensions.Y = FMath::Abs((MediaPlateRenderTarget->SizeY / 2) - FMath::Floor(MousePosition.Y)) * 2;
+
+	// Early-out if the dimensions are invalid
+	if ((CameraFeedDimensions.X == 0) || (CameraFeedDimensions.Y == 0))
+	{
+		return;
+	}
 
-			const FIntPoint CGResolution = CGLayer->GetRenderResolution();
-			const float CGLayerAspectRatio = (CGResolution.Y != 0) ? (float)CGResolution.X / (float)CGResolution.Y : 0.0f;
+	if (ULensFile* LensFilePtr = GetLensFile())
+	{
+		// Compare the aspect ratio of the selected camera to the cg layer's aspect ratio to determine if the cg layer needs to be resized
+		if (CineCameraComponent.IsValid())
+		{
+			const float CameraFeedAspectRatio = CameraFeedDimensions.X / (float)CameraFeedDimensions.Y;
+			const float CameraAspectRatio = CineCameraComponent->Filmback.SensorAspectRatio;
 
-			if (CGLayerAspectRatio > MediaAspectRatio)
+			// If the two aspect ratios are within the acceptable tolerance, attempt to minimze the aspect ratio difference by adjusting the camera feed dimensions
+			constexpr float AspectRatioNudgeTolerance = 0.1f;
+			if (FMath::IsNearlyEqual(CameraAspectRatio, CameraFeedAspectRatio, AspectRatioNudgeTolerance))
 			{
-				const float AspectRatioCorrection = (MediaAspectRatio != 0.0f) ? CGLayerAspectRatio / MediaAspectRatio : 1.0f;
-				MaterialPass->Material.SetScalarOverride(TEXT("AspectRatioCorrection_H"), 1.0f);
-				MaterialPass->Material.SetScalarOverride(TEXT("AspectRatioCorrection_V"), AspectRatioCorrection);
+				MinimizeAspectRatioError(CameraFeedDimensions, CameraAspectRatio);
 			}
-			else if (MediaAspectRatio > CGLayerAspectRatio)
+
+			LensFilePtr->CameraFeedInfo.SetDimensions(CameraFeedDimensions);
+
+			UpdateAspectRatioCorrection(CameraFeedDimensions);
+		}
+	}
+}
+
+void FCameraCalibrationStepsController::UpdateAspectRatioCorrection(FIntPoint CameraFeedDimensions)
+{
+	if (!MaterialPass.IsValid())
+	{
+		return;
+	}
+
+	// If the camera feed dimensions are invalid, use the ratios of the media and CG layers to set the aspect ratio correction material parameters
+	if ((CameraFeedDimensions.X == 0) || (CameraFeedDimensions.Y == 0))
+	{
+		const float MediaAspectRatio = (MediaPlateRenderTarget->SizeY != 0) ? (float)MediaPlateRenderTarget->SizeX / (float)MediaPlateRenderTarget->SizeY : 0.0f;
+
+		const FIntPoint CGResolution = CGLayer->GetRenderResolution();
+		const float CGLayerAspectRatio = (CGResolution.Y != 0) ? (float)CGResolution.X / (float)CGResolution.Y : 0.0f;
+
+		if (CGLayerAspectRatio > MediaAspectRatio)
+		{
+			const float AspectRatioCorrection = (MediaAspectRatio != 0.0f) ? CGLayerAspectRatio / MediaAspectRatio : 1.0f;
+			MaterialPass->Material.SetScalarOverride(TEXT("AspectRatioCorrection_H"), 1.0f);
+			MaterialPass->Material.SetScalarOverride(TEXT("AspectRatioCorrection_V"), AspectRatioCorrection);
+		}
+		else if (MediaAspectRatio > CGLayerAspectRatio)
+		{
+			const float AspectRatioCorrection = (CGLayerAspectRatio != 0.0f) ? MediaAspectRatio / CGLayerAspectRatio : 1.0f;
+			MaterialPass->Material.SetScalarOverride(TEXT("AspectRatioCorrection_H"), AspectRatioCorrection);
+			MaterialPass->Material.SetScalarOverride(TEXT("AspectRatioCorrection_V"), 1.0f);
+		}
+		else
+		{
+			MaterialPass->Material.SetScalarOverride(TEXT("AspectRatioCorrection_H"), 1.0f);
+			MaterialPass->Material.SetScalarOverride(TEXT("AspectRatioCorrection_V"), 1.0f);
+		}
+	}
+	else
+	{
+		// If the camera feed dimensions are valid, use them and the dimensions of the media to set the aspect ratio correction material parameters
+		MaterialPass->Material.SetScalarOverride(TEXT("AspectRatioCorrection_H"), MediaPlateRenderTarget->SizeX / (float)CameraFeedDimensions.X);
+		MaterialPass->Material.SetScalarOverride(TEXT("AspectRatioCorrection_V"), MediaPlateRenderTarget->SizeY / (float)CameraFeedDimensions.Y);
+	}
+}
+
+void FCameraCalibrationStepsController::MinimizeAspectRatioError(FIntPoint& CameraFeedDimensions, float CameraAspectRatio)
+{
+	float CameraFeedAspectRatio = (CameraFeedDimensions.Y > 0) ? CameraFeedDimensions.X / (float)CameraFeedDimensions.Y : 0.0f;
+
+	// Track the number of iterations to ensure that the minimization does not continue infinitely
+	int32 NumIterations = 0;
+	constexpr int32 MaxIterations = 20;
+
+	constexpr float AspectRatioErrorTolerance = 0.001f;
+	while ((!FMath::IsNearlyEqual(CameraAspectRatio, CameraFeedAspectRatio, AspectRatioErrorTolerance)) && NumIterations < MaxIterations)
+	{
+		// The camera feed needs to be wider and/or shorter
+		if (CameraFeedAspectRatio < CameraAspectRatio)
+		{
+			// Use modulus to determine whether to alter the width or height each iteration
+			if (NumIterations % 2 == 0)
 			{
-				const float AspectRatioCorrection = (CGLayerAspectRatio != 0.0f) ? MediaAspectRatio / CGLayerAspectRatio : 1.0f;
-				MaterialPass->Material.SetScalarOverride(TEXT("AspectRatioCorrection_H"), AspectRatioCorrection);
-				MaterialPass->Material.SetScalarOverride(TEXT("AspectRatioCorrection_V"), 1.0f);
+				CameraFeedDimensions.X += 2;
 			}
 			else
 			{
-				MaterialPass->Material.SetScalarOverride(TEXT("AspectRatioCorrection_H"), 1.0f);
-				MaterialPass->Material.SetScalarOverride(TEXT("AspectRatioCorrection_V"), 1.0f);
+				CameraFeedDimensions.Y -= 2;
 			}
 		}
-	}
-
-	for (TStrongObjectPtr<UCameraCalibrationStep>& Step : CalibrationSteps)
-	{
-		if (Step.IsValid())
+		else // The camera feed needs to be narrower or taller
 		{
-			Step->Tick(DeltaTime);
+			// Use modulus to determine whether to change the width or height each iteration
+			if (NumIterations % 2 == 0)
+			{
+				CameraFeedDimensions.X -= 2;
+			}
+			else
+			{
+				CameraFeedDimensions.Y += 2;
+			}
 		}
-	}
 
-	return true;
+		CameraFeedAspectRatio = (CameraFeedDimensions.Y > 0) ? CameraFeedDimensions.X / (float)CameraFeedDimensions.Y : 0.0f;
+
+		NumIterations++;
+	}
 }
 
 TWeakObjectPtr<ACompositingElement> FCameraCalibrationStepsController::AddElement(ACompositingElement* Parent, FString& ClassPath, FString& ElementName) const
diff --git a/Engine/Plugins/VirtualProduction/CameraCalibration/Source/CameraCalibrationEditor/Private/AssetEditor/CameraCalibrationStepsController.h b/Engine/Plugins/VirtualProduction/CameraCalibration/Source/CameraCalibrationEditor/Private/AssetEditor/CameraCalibrationStepsController.h
--- a/Engine/Plugins/VirtualProduction/CameraCalibration/Source/CameraCalibrationEditor/Private/AssetEditor/CameraCalibrationStepsController.h
+++ b/Engine/Plugins/VirtualProduction/CameraCalibration/Source/CameraCalibrationEditor/Private/AssetEditor/CameraCalibrationStepsController.h
@@ -137,6 +137,12 @@ public:
 	/** Redraw the overlay material used by the input overlay pass */
 	void RefreshOverlay(EOverlayPassType OverlayPass = EOverlayPassType::ToolOverlay);
 
+	/** Use the input mouse position (representing any corner of the camera feed) and the dimensions of the media source to calculate the new size of the camera feed */
+	void UpdateCameraFeedInfo(FVector2D MousePosition);
+
+	/** Update the material parameters that control the CG layer's aspect ratio correction based on the camera feed (if it is valid) or on the aspect ratio of the camera actor */
+	void UpdateAspectRatioCorrection(FIntPoint CameraFeedDimensions);
+
 public:
 
 	/** Called by the UI when the Simulcam Viewport is clicked */
@@ -228,6 +234,9 @@ private:
 	/** Returns the overlay material used by the input overlay pass type */
 	UMaterialInterface* GetOverlayMaterial(EOverlayPassType OverlayPass) const;
 
+	/** Iteratively attempt to resize the camera feed until its aspect ratio matches the input camera aspect ratio */
+	void MinimizeAspectRatioError(FIntPoint& CameraFeedDimensions, float CameraAspectRatio);
+
 private:
 
 	/** Pointer to the camera calibration toolkit */
diff --git a/Engine/Plugins/VirtualProduction/CameraCalibration/Source/CameraCalibrationEditor/Private/AssetEditor/LensInfoStep.cpp b/Engine/Plugins/VirtualProduction/CameraCalibration/Source/CameraCalibrationEditor/Private/AssetEditor/LensInfoStep.cpp
--- a/Engine/Plugins/VirtualProduction/CameraCalibration/Source/CameraCalibrationEditor/Private/AssetEditor/LensInfoStep.cpp
+++ b/Engine/Plugins/VirtualProduction/CameraCalibration/Source/CameraCalibrationEditor/Private/AssetEditor/LensInfoStep.cpp
@@ -11,6 +11,7 @@
 #include "ScopedTransaction.h"
 #include "SResetToDefaultMenu.h"
 #include "Widgets/Input/SButton.h"
+#include "Widgets/Text/STextBlock.h"
 
 
 #define LOCTEXT_NAMESPACE "LensInfoStep"
@@ -46,6 +47,15 @@ TSharedRef<SWidget> ULensInfoStep::BuildUI()
 		}
 	});
 
+	/** Camera Feed Info Details View */
+	FStructureDetailsViewArgs CameraFeedInfoStructDetailsView;
+
+	TSharedRef<FStructOnScope> CameraFeedInfoStructOnScope = MakeShared<FStructOnScope>(FCameraFeedInfo::StaticStruct(), reinterpret_cast<uint8*>(&LensFile->CameraFeedInfo));
+	TSharedPtr<IStructureDetailsView> CameraFeedInfoStructureDetailsView = PropertyEditor.CreateStructureDetailView(DetailArgs, CameraFeedInfoStructDetailsView, CameraFeedInfoStructOnScope);
+
+	CameraFeedInfoStructureDetailsView->GetDetailsView()->OnFinishedChangingProperties().AddUObject(this, &ULensInfoStep::OnCameraFeedInfoChanged);
+
+	/** Simulcam Info Details View */
 	FStructureDetailsViewArgs SimulcamInfoStructDetailsView;
 
 	TSharedRef<FStructOnScope> SimulcamInfoStructOnScope = MakeShared<FStructOnScope>(FSimulcamInfo::StaticStruct(), reinterpret_cast<uint8*>(&LensFile->SimulcamInfo));
@@ -72,7 +82,26 @@ TSharedRef<SWidget> ULensInfoStep::BuildUI()
 				.DiffersFromDefault(TAttribute<bool>::Create(TAttribute<bool>::FGetter::CreateUObject(this, &ULensInfoStep::DiffersFromDefault)))
 			]
 		]
-	
+
+		+ SVerticalBox::Slot() // Camera Feed information structure
+		.AutoHeight()
+		[
+			SNew(SHorizontalBox)
+			+ SHorizontalBox::Slot()
+			[CameraFeedInfoStructureDetailsView->GetWidget().ToSharedRef() ]
+		]
+
+		+ SVerticalBox::Slot()
+		.AutoHeight()
+		.Padding(FMargin(4.0f, 4.0f, 4.0f, 4.0f))
+		.HAlign(HAlign_Center)
+		[
+			SNew(STextBlock)
+			.Text(LOCTEXT("ControllerWarningToolTip", "The aspect ratios of the camera feed and the CG camera do not match."))
+			.ColorAndOpacity(FSlateColor(FLinearColor::Yellow))
+			.Visibility_UObject(this, &ULensInfoStep::HandleAspectRatioWarningVisibility)
+		]
+
 		+ SVerticalBox::Slot() // Simulcam information structure
 		.AutoHeight()
 		[
@@ -173,6 +202,65 @@ void ULensInfoStep::OnSaveLensInformation()
 	}
 }
 
+bool ULensInfoStep::OnViewportClicked(const FGeometry& MyGeometry, const FPointerEvent& MouseEvent)
+{
+	// User interaction is Ctrl + Left Mouse Button
+	if ((MouseEvent.GetEffectingButton() == EKeys::LeftMouseButton) && (MouseEvent.GetModifierKeys().IsControlDown()))
+	{
+		if (TSharedPtr<FCameraCalibrationStepsController> StepsController = CameraCalibrationStepsController.Pin())
+		{
+			const FVector2D LocalInPixels = MyGeometry.AbsoluteToLocal(MouseEvent.GetScreenSpacePosition());
+			StepsController->UpdateCameraFeedInfo(LocalInPixels);
+			return true;
+		}
+	}
+
+	return false;
+}
+
+void ULensInfoStep::OnCameraFeedInfoChanged(const FPropertyChangedEvent& PropertyChangedEvent)
+{
+	// Ignore temporary interaction (dragging sliders, etc.)
+	if (PropertyChangedEvent.ChangeType == EPropertyChangeType::ValueSet)
+	{
+		if (TSharedPtr<FCameraCalibrationStepsController> StepsController = CameraCalibrationStepsController.Pin())
+		{
+			if (ULensFile* LensFile = StepsController->GetLensFile())
+			{
+				LensFile->CameraFeedInfo.SetDimensions(LensFile->CameraFeedInfo.GetDimensions());
+				StepsController->UpdateAspectRatioCorrection(LensFile->CameraFeedInfo.GetDimensions());
+			}
+		}
+	}
+}
+
+EVisibility ULensInfoStep::HandleAspectRatioWarningVisibility() const
+{
+	if (TSharedPtr<FCameraCalibrationStepsController> StepsController = CameraCalibrationStepsController.Pin())
+	{
+		if (ULensFile* LensFile = CameraCalibrationStepsController.Pin()->GetLensFile())
+		{
+			// Do not display the warning if the camera feed info is invalid
+			if (!LensFile->CameraFeedInfo.IsValid())
+			{
+				return EVisibility::Collapsed;
+			}
+
+			const float CameraFeedAspectRatio = LensFile->CameraFeedInfo.GetAspectRatio();
+			const float CGCameraAspectRatio = LensFile->SimulcamInfo.CGLayerAspectRatio;
+
+			// Display the warning if the difference between the two aspect ratios is higher than the acceptable tolerance
+			constexpr float AspectRatioErrorTolerance = 0.01f;
+			if (!FMath::IsNearlyEqual(CameraFeedAspectRatio, CGCameraAspectRatio, AspectRatioErrorTolerance))
+			{
+				return EVisibility::Visible;
+			}
+		}
+	}
+
+	return EVisibility::Collapsed;
+}
+
 bool ULensInfoStep::IsActive() const
 {
 	return bIsActive;
diff --git a/Engine/Plugins/VirtualProduction/CameraCalibration/Source/CameraCalibrationEditor/Private/AssetEditor/LensInfoStep.h b/Engine/Plugins/VirtualProduction/CameraCalibration/Source/CameraCalibrationEditor/Private/AssetEditor/LensInfoStep.h
--- a/Engine/Plugins/VirtualProduction/CameraCalibration/Source/CameraCalibrationEditor/Private/AssetEditor/LensInfoStep.h
+++ b/Engine/Plugins/VirtualProduction/CameraCalibration/Source/CameraCalibrationEditor/Private/AssetEditor/LensInfoStep.h
@@ -21,6 +21,7 @@ public:
 
 	//~ Begin UCameraCalibrationStep interface
 	virtual void Initialize(TWeakPtr<FCameraCalibrationStepsController> InCameraCalibrationStepController) override;
+	virtual bool OnViewportClicked(const FGeometry& MyGeometry, const FPointerEvent& MouseEvent) override;
 	virtual TSharedRef<SWidget> BuildUI() override;
 	virtual FName FriendlyName() const  override { return TEXT("Lens Information"); };
 	virtual bool DependsOnStep(UCameraCalibrationStep* Step) const override;
@@ -34,12 +35,19 @@ public:
 	/** Triggered when user edits any of the LensInfo properties */
 	void OnSaveLensInformation();
 
+	/** Triggered when user edits any of the CameraFeedInfo properties */
+	void OnCameraFeedInfoChanged(const FPropertyChangedEvent& PropertyChangedEvent);
+
 	/** Reapply the original LensInfo to the LensFile */
 	void ResetToDefault();
 
 	/** Returns true if our LensInfo differs from the original info when the tool was opened */
 	bool DiffersFromDefault() const;
 
+private:
+	/** Determines the visibility of the aspect ratio warning by comparing the camera feed aspect ratio with that of the camera actor */
+	EVisibility HandleAspectRatioWarningVisibility() const;
+
 private:
 
 	/** Pointer to the calibration steps controller */
diff --git a/Engine/Plugins/VirtualProduction/CameraCalibrationCore/Source/CameraCalibrationCore/Public/LensData.h b/Engine/Plugins/VirtualProduction/CameraCalibrationCore/Source/CameraCalibrationCore/Public/LensData.h
--- a/Engine/Plugins/VirtualProduction/CameraCalibrationCore/Source/CameraCalibrationCore/Public/LensData.h
+++ b/Engine/Plugins/VirtualProduction/CameraCalibrationCore/Source/CameraCalibrationCore/Public/LensData.h
@@ -60,6 +60,68 @@ public:
 
 };
 
+/**
+ * Information about a camera feed, including its dimensions and aspect ratio.
+ * The camera feed represents (potentially) an inner rect of a media input, excluding any masks / extractions that may surround it.
+ */
+USTRUCT(BlueprintType)
+struct CAMERACALIBRATIONCORE_API FCameraFeedInfo
+{
+	GENERATED_BODY()
+
+public:
+	FIntPoint GetDimensions()
+	{
+		return Dimensions;
+	}
+
+	void SetDimensions(FIntPoint InDimensions)
+	{
+		Dimensions = InDimensions;
+
+		// If one of the dimensions is set to 0, the description of the feed is invalid
+		if (Dimensions.X == 0 || Dimensions.Y == 0)
+		{
+			AspectRatio = 0.0f;
+			bIsValid = false;
+		}
+		else
+		{
+			AspectRatio = Dimensions.X / (float)Dimensions.Y;
+			bIsValid = true;
+		}
+	}
+
+	float GetAspectRatio()
+	{
+		return AspectRatio;
+	}
+
+	void Invalidate()
+	{
+		Dimensions = FIntPoint(0, 0);
+		AspectRatio = 0.0f;
+		bIsValid = false;
+	}
+
+	bool IsValid()
+	{
+		return bIsValid;
+	}
+
+private:
+	/** Dimensions of the Camera Feed */
+	UPROPERTY(EditAnywhere, Category = "Camera Feed Info")
+	FIntPoint Dimensions = FIntPoint(0, 0);
+
+	/** Aspect Ratio of the Camera Feed */
+	UPROPERTY(VisibleAnywhere, Category = "Camera Feed Info")
+	float AspectRatio = 0.0f;
+
+	/** Whether the dimensions of the feed are valid (false if either X or Y is 0) */
+	bool bIsValid = false;
+};
+
 /**
  * Information about the simulcam composition
  */
diff --git a/Engine/Plugins/VirtualProduction/CameraCalibrationCore/Source/CameraCalibrationCore/Public/LensFile.h b/Engine/Plugins/VirtualProduction/CameraCalibrationCore/Source/CameraCalibrationCore/Public/LensFile.h
--- a/Engine/Plugins/VirtualProduction/CameraCalibrationCore/Source/CameraCalibrationCore/Public/LensFile.h
+++ b/Engine/Plugins/VirtualProduction/CameraCalibrationCore/Source/CameraCalibrationCore/Public/LensFile.h
@@ -238,9 +238,15 @@ public:
 	UPROPERTY(EditAnywhere, BlueprintReadWrite, Category = "Lens info")
 	FLensInfo LensInfo;
 
+#if WITH_EDITORONLY_DATA
+	/** Camera feed information */
+	UPROPERTY(EditAnywhere, BlueprintReadWrite, Category = "Camera Feed")
+	FCameraFeedInfo CameraFeedInfo;
+
 	/** Simulcam information */
 	UPROPERTY(VisibleAnywhere, BlueprintReadOnly, Category = "Simulcam info")
 	FSimulcamInfo SimulcamInfo;
+#endif // WITH_EDITORONLY_DATA
 
 	/** Type of data used for lens mapping */
 	UPROPERTY(EditAnywhere, BlueprintReadWrite, Category = "Lens mapping")
