CameraCalibration: Fix the way the calibration tools calculate the clicked mouse position in the simulcam viewport when the camera feed resolution and media resolution are different

diff --git a/Engine/Plugins/Runtime/OpenCV/Source/OpenCVHelper/Private/OpenCVHelper.cpp b/Engine/Plugins/Runtime/OpenCV/Source/OpenCVHelper/Private/OpenCVHelper.cpp
--- a/Engine/Plugins/Runtime/OpenCV/Source/OpenCVHelper/Private/OpenCVHelper.cpp
+++ b/Engine/Plugins/Runtime/OpenCV/Source/OpenCVHelper/Private/OpenCVHelper.cpp
@@ -241,7 +241,7 @@ UTexture2D* FOpenCVHelper::TextureFromCvMat(cv::Mat& Mat, UTexture2D* InTexture)
 
 	if ((InTexture->GetSizeX() != Mat.cols) || (InTexture->GetSizeY() != Mat.rows) || (InTexture->GetPixelFormat() != PixelFormat))
 	{
-		return nullptr;
+		return TextureFromCvMat(Mat);
 	}
 
 	// Copy the pixels from the OpenCV Mat to the Texture
diff --git a/Engine/Plugins/VirtualProduction/CameraCalibration/Source/CameraCalibrationEditor/Private/AssetEditor/CameraCalibrationStepsController.cpp b/Engine/Plugins/VirtualProduction/CameraCalibration/Source/CameraCalibrationEditor/Private/AssetEditor/CameraCalibrationStepsController.cpp
--- a/Engine/Plugins/VirtualProduction/CameraCalibration/Source/CameraCalibrationEditor/Private/AssetEditor/CameraCalibrationStepsController.cpp
+++ b/Engine/Plugins/VirtualProduction/CameraCalibration/Source/CameraCalibrationEditor/Private/AssetEditor/CameraCalibrationStepsController.cpp
@@ -77,7 +77,6 @@ namespace CameraCalibrationStepsController
 
 FCameraCalibrationStepsController::FCameraCalibrationStepsController(TWeakPtr<FCameraCalibrationToolkit> InCameraCalibrationToolkit, ULensFile* InLensFile)
 	: CameraCalibrationToolkit(InCameraCalibrationToolkit)
-	, RenderTargetSize(FIntPoint(1920, 1080))
 	, LensFile(TWeakObjectPtr<ULensFile>(InLensFile))
 {
 	check(CameraCalibrationToolkit.IsValid());
@@ -176,96 +175,166 @@ TSharedPtr<SWidget> FCameraCalibrationStepsController::BuildUI()
 
 bool FCameraCalibrationStepsController::OnTick(float DeltaTime)
 {
-	// Update the lens file eval data
+	// Update the lens file evaluation inputs
 	LensFileEvaluationInputs.bIsValid = false;
 	if (const ULensComponent* const LensComponent = FindLensComponent())
 	{
 		LensFileEvaluationInputs = LensComponent->GetLensFileEvaluationInputs();
 	}
 
-	bool bCompElementsResized = false;
+	// Update the output resolution of the comp to match the resolution of the media source (if it has changed)
+	const bool bCompResized = UpdateCompResolution();
 
-	// Compare the current dimensions of the playing media track to the comp's render resolution to determine if the comp needs to be resized
-	if (MediaPlayer.IsValid())
+	// Update the output resolution of the CG Layer if either the Comp was resized or if the camera's aspect ratio has changed
+	UpdateCGResolution(bCompResized);
+
+	// Tick each of the calibration steps
+	for (TStrongObjectPtr<UCameraCalibrationStep>& Step : CalibrationSteps)
+	{
+		if (Step.IsValid())
+		{
+			Step->Tick(DeltaTime);
+		}
+	}
+
+	return true;
+}
+
+bool FCameraCalibrationStepsController::UpdateCompResolution()
+{
+	const UMediaPlayer* const MediaPlayerPtr = MediaPlayer.Get();
+
+	if (!MediaPlayerPtr)
 	{
-		const FIntPoint MediaDimensions = MediaPlayer->GetVideoTrackDimensions(INDEX_NONE, INDEX_NONE);
+		return false;
+	}
+
+	const FIntPoint MediaDimensions = MediaPlayerPtr->GetVideoTrackDimensions(INDEX_NONE, INDEX_NONE);
+	const FIntPoint CompRenderResolution = GetCompRenderResolution();
 
-		// If no track was found, the dimensions might be (0, 0)
-		if (MediaDimensions.X != 0 && MediaDimensions.Y != 0)
+	bool bCompResized = false;
+
+	// If no track was found, the dimensions might be (0, 0)
+	if (MediaDimensions.X != 0 && MediaDimensions.Y != 0)
+	{
+		if (MediaDimensions != CompRenderResolution)
 		{
-			if (RenderTarget->SizeX != MediaDimensions.X || RenderTarget->SizeY != MediaDimensions.Y)
+			// Resize the output resolution of the top-level comp to match the incoming media dimensions
+			if (ACompositingElement* CompPtr = Comp.Get())
 			{
-				// Resize the media plate comp layer and its output render target to match the incoming media dimensions
-				MediaPlateRenderTarget->ResizeTarget(MediaDimensions.X, MediaDimensions.Y);
-				MediaPlate->SetRenderResolution(MediaDimensions);
-
-				// Resize the parent comp and its output render target to match the incoming media dimensions
-				RenderTarget->ResizeTarget(MediaDimensions.X, MediaDimensions.Y);
 				Comp->SetRenderResolution(MediaDimensions);
+			}
 
-				bCompElementsResized = true;
+			// Resize the output render targets to match the incoming media dimensions
+			if (UTextureRenderTarget2D* RenderTargetPtr = RenderTarget.Get())
+			{
+				RenderTargetPtr->ResizeTarget(MediaDimensions.X, MediaDimensions.Y);
 			}
+
+			if (UTextureRenderTarget2D* MediaPlateRenderTargetPtr = MediaPlateRenderTarget.Get())
+			{
+				MediaPlateRenderTargetPtr->ResizeTarget(MediaDimensions.X, MediaDimensions.Y);
+			}
+
+			if (UTextureRenderTarget2D* ToolOverlayRenderTargetPtr = ToolOverlayRenderTarget.Get())
+			{
+				ToolOverlayRenderTargetPtr->ResizeTarget(MediaDimensions.X, MediaDimensions.Y);
+			}
+
+			if (UTextureRenderTarget2D* UserOverlayRenderTargetPtr = UserOverlayRenderTarget.Get())
+			{
+				UserOverlayRenderTargetPtr->ResizeTarget(MediaDimensions.X, MediaDimensions.Y);
+			}
+
+			bCompResized = true;
 		}
+	}
 
-		if (ULensFile* LensFilePtr = GetLensFile())
+	// Update the LensFile SimulcamInfo if the Media Dimensions have changed
+	if (ULensFile* const LensFilePtr = GetLensFile())
+	{
+		if (LensFilePtr->SimulcamInfo.MediaResolution != MediaDimensions)
 		{
 			LensFilePtr->SimulcamInfo.MediaResolution = MediaDimensions;
 			LensFilePtr->SimulcamInfo.MediaPlateAspectRatio = (MediaDimensions.Y != 0) ? (float)MediaDimensions.X / (float)MediaDimensions.Y : 0.0f;
 		}
 	}
 
-	// Compare the aspect ratio of the selected camera to the cg layer's aspect ratio to determine if the cg layer needs to be resized
-	if (CineCameraComponent.Get())
-	{
-		const float CameraAspectRatio = CineCameraComponent->Filmback.SensorAspectRatio;
+	return bCompResized;
+}
 
-		FIntPoint CGResolution = CGLayer->GetRenderResolution();
-		const float CGLayerAspectRatio = (CGResolution.Y != 0) ? (float)CGResolution.X / (float)CGResolution.Y : 0.0f;
+void FCameraCalibrationStepsController::UpdateCGResolution(bool bCompResized)
+{
+	const UCineCameraComponent* const CineCameraComponentPtr = CineCameraComponent.Get();
+	const ACompositingElement* const CompPtr = Comp.Get();
+	ACompositingElement* const CGLayerPtr = CGLayer.Get();
+	ULensFile* const LensFilePtr = GetLensFile();
 
-		constexpr float Tolerance = 0.01;
-		if (!FMath::IsNearlyEqual(CameraAspectRatio, CGLayerAspectRatio, Tolerance))
-		{
-			CGResolution.X = CGResolution.Y * CameraAspectRatio;
+	if (!CineCameraComponentPtr || !CompPtr || !CGLayerPtr || !LensFilePtr)
+	{
+		return;
+	}
 
-			CGLayer->ResolutionSource = EInheritedSourceType::Override;
-			CGLayer->SetRenderResolution(CGResolution);
+	// There are two cases where the CG Layer's output resolution might need to be resized:
+	// 1) If the aspect ratio of the camera source's filmback changed, to avoid any stretching of the CG
+	// 2) If the Comp's output resolution changed, to ensure that the CG layer's resolution is not larger than the Comp, and not unnecessarily smaller
 
-			bCompElementsResized = true;
-		}
+	const FIntPoint CGResolution = CGLayerPtr->GetRenderResolution();
+	const float CGLayerAspectRatio = (CGResolution.Y != 0) ? (float)CGResolution.X / (float)CGResolution.Y : 0.0f;
+	const float FilmbackAspectRatio = CineCameraComponentPtr->Filmback.SensorAspectRatio;
 
-		if (ULensFile* LensFilePtr = GetLensFile())
-		{
-			LensFilePtr->SimulcamInfo.CGLayerAspectRatio = CameraAspectRatio;
-		}
+	// Early-out if the filmback is somehow 0.0f (which shouldn't be possible)
+	if (FMath::IsNearlyEqual(FilmbackAspectRatio, 0.0f))
+	{
+		return;
 	}
 
-	// If either the CG or the Media was resized, then recompute the aspect ratio correction for the comp material
-	if (bCompElementsResized)
+	constexpr float Tolerance = 0.01f;
+	const bool bFilmbackAspectRatioChanged = !FMath::IsNearlyEqual(FilmbackAspectRatio, CGLayerAspectRatio, Tolerance);
+
+	if (bFilmbackAspectRatioChanged || bCompResized)
 	{
-		if (ULensFile* LensFilePtr = GetLensFile())
+		const FIntPoint CompResolution = GetCompRenderResolution();
+		const float CompAspectRatio = (CompResolution.Y != 0) ? (float)CompResolution.X / (float)CompResolution.Y : 0.0f;
+
+		// Compute the camera feed dimensions that would fit the aspect ratio of the camera within the comp
+		// If the filmback aspect ratio is wider than the comp, then the width of the camera feed will equal the width of the comp, and the height will be scaled
+		// If the filmback aspect ratio is narrower than the comp, then the height of the camera feed will equal the height of the comp, and the width will be scaled
+
+		FIntPoint NewCGResolution = CompResolution;
+
+		if (FilmbackAspectRatio > CompAspectRatio)
 		{
-			UpdateAspectRatioCorrection(LensFilePtr->CameraFeedInfo.GetDimensions());
+			NewCGResolution.Y = CompResolution.X / FilmbackAspectRatio;
+		}
+		else if (CompAspectRatio > FilmbackAspectRatio)
+		{
+			NewCGResolution.X = CompResolution.Y * FilmbackAspectRatio;
 		}
-	}
 
-	for (TStrongObjectPtr<UCameraCalibrationStep>& Step : CalibrationSteps)
-	{
-		if (Step.IsValid())
+		CGLayerPtr->SetRenderResolution(NewCGResolution);
+
+		LensFilePtr->SimulcamInfo.CGLayerAspectRatio = FilmbackAspectRatio;
+
+		if (!LensFilePtr->CameraFeedInfo.IsOverridden())
 		{
-			Step->Tick(DeltaTime);
+			constexpr bool bMarkAsOverridden = false;
+			LensFilePtr->CameraFeedInfo.SetDimensions(NewCGResolution, bMarkAsOverridden);
 		}
-	}
 
-	return true;
+		UpdateAspectRatioCorrection();
+	}
 }
 
-void FCameraCalibrationStepsController::UpdateCameraFeedInfo(FVector2D MousePosition)
+void FCameraCalibrationStepsController::SetCameraFeedDimensionsFromMousePosition(FVector2D MousePosition)
 {
 	// Compute the size of the camera feed, using the MousePosition as one of its corners.
 	// We assume that the camera feed will always be centered in the media. If that assumption proves false, this math should be updated in the future.
+	const FIntPoint CompRenderResolution = GetCompRenderResolution();
+
 	FIntPoint CameraFeedDimensions = FIntPoint(0, 0);
-	CameraFeedDimensions.X = FMath::Abs((MediaPlateRenderTarget->SizeX / 2) - FMath::Floor(MousePosition.X)) * 2;
-	CameraFeedDimensions.Y = FMath::Abs((MediaPlateRenderTarget->SizeY / 2) - FMath::Floor(MousePosition.Y)) * 2;
+	CameraFeedDimensions.X = FMath::Abs((CompRenderResolution.X / 2) - FMath::Floor(MousePosition.X)) * 2;
+	CameraFeedDimensions.Y = FMath::Abs((CompRenderResolution.Y / 2) - FMath::Floor(MousePosition.Y)) * 2;
 
 	// Early-out if the dimensions are invalid
 	if ((CameraFeedDimensions.X == 0) || (CameraFeedDimensions.Y == 0))
@@ -273,7 +342,7 @@ void FCameraCalibrationStepsController::UpdateCameraFeedInfo(FVector2D MousePosi
 		return;
 	}
 
-	if (ULensFile* LensFilePtr = GetLensFile())
+	if (ULensFile* const LensFilePtr = GetLensFile())
 	{
 		// Compare the aspect ratio of the selected camera to the cg layer's aspect ratio to determine if the cg layer needs to be resized
 		if (CineCameraComponent.IsValid())
@@ -281,59 +350,58 @@ void FCameraCalibrationStepsController::UpdateCameraFeedInfo(FVector2D MousePosi
 			const float CameraFeedAspectRatio = CameraFeedDimensions.X / (float)CameraFeedDimensions.Y;
 			const float CameraAspectRatio = CineCameraComponent->Filmback.SensorAspectRatio;
 
-			// If the two aspect ratios are within the acceptable tolerance, attempt to minimze the aspect ratio difference by adjusting the camera feed dimensions
+			// If the two aspect ratios are within the acceptable tolerance, attempt to minimize the aspect ratio difference by adjusting the camera feed dimensions
 			constexpr float AspectRatioNudgeTolerance = 0.1f;
 			if (FMath::IsNearlyEqual(CameraAspectRatio, CameraFeedAspectRatio, AspectRatioNudgeTolerance))
 			{
 				MinimizeAspectRatioError(CameraFeedDimensions, CameraAspectRatio);
 			}
 
-			LensFilePtr->CameraFeedInfo.SetDimensions(CameraFeedDimensions);
+			// Update the camera feed dimensions and mark as overridden because it was changed as the result of user interaction
+			constexpr bool bMarkAsOverridden = true;
+			LensFilePtr->CameraFeedInfo.SetDimensions(CameraFeedDimensions, bMarkAsOverridden);
 
-			UpdateAspectRatioCorrection(CameraFeedDimensions);
+			UpdateAspectRatioCorrection();
 		}
 	}
 }
 
-void FCameraCalibrationStepsController::UpdateAspectRatioCorrection(FIntPoint CameraFeedDimensions)
+void FCameraCalibrationStepsController::SetCameraFeedDimensions(FIntPoint Dimensions, bool bMarkAsOverridden)
 {
-	if (!MaterialPass.IsValid())
+	if (ULensFile* const LensFilePtr = GetLensFile())
 	{
-		return;
+		LensFile->CameraFeedInfo.SetDimensions(Dimensions, bMarkAsOverridden);
+		UpdateAspectRatioCorrection();
 	}
+}
 
-	// If the camera feed dimensions are invalid, use the ratios of the media and CG layers to set the aspect ratio correction material parameters
-	if ((CameraFeedDimensions.X == 0) || (CameraFeedDimensions.Y == 0))
+void FCameraCalibrationStepsController::UpdateAspectRatioCorrection()
+{
+	const ULensFile* const LensFilePtr = LensFile.Get();
+	const ACompositingElement* const CompPtr = Comp.Get();
+	UCompositingElementMaterialPass* const MaterialPassPtr = MaterialPass.Get();
+
+	if (!MaterialPassPtr || !CompPtr || !LensFilePtr)
 	{
-		const float MediaAspectRatio = (MediaPlateRenderTarget->SizeY != 0) ? (float)MediaPlateRenderTarget->SizeX / (float)MediaPlateRenderTarget->SizeY : 0.0f;
+		return;
+	}
 
-		const FIntPoint CGResolution = CGLayer->GetRenderResolution();
-		const float CGLayerAspectRatio = (CGResolution.Y != 0) ? (float)CGResolution.X / (float)CGResolution.Y : 0.0f;
+	float AspectRatioCorrectionX = 1.0f;
+	float AspectRatioCorrectionY = 1.0f;
 
-		if (CGLayerAspectRatio > MediaAspectRatio)
-		{
-			const float AspectRatioCorrection = (MediaAspectRatio != 0.0f) ? CGLayerAspectRatio / MediaAspectRatio : 1.0f;
-			MaterialPass->Material.SetScalarOverride(TEXT("AspectRatioCorrection_H"), 1.0f);
-			MaterialPass->Material.SetScalarOverride(TEXT("AspectRatioCorrection_V"), AspectRatioCorrection);
-		}
-		else if (MediaAspectRatio > CGLayerAspectRatio)
-		{
-			const float AspectRatioCorrection = (CGLayerAspectRatio != 0.0f) ? MediaAspectRatio / CGLayerAspectRatio : 1.0f;
-			MaterialPass->Material.SetScalarOverride(TEXT("AspectRatioCorrection_H"), AspectRatioCorrection);
-			MaterialPass->Material.SetScalarOverride(TEXT("AspectRatioCorrection_V"), 1.0f);
-		}
-		else
-		{
-			MaterialPass->Material.SetScalarOverride(TEXT("AspectRatioCorrection_H"), 1.0f);
-			MaterialPass->Material.SetScalarOverride(TEXT("AspectRatioCorrection_V"), 1.0f);
-		}
-	}
-	else
+	// If the camera feed dimensions are valid, use them and the dimensions of the media to set the aspect ratio correction material parameters
+	if (LensFilePtr->CameraFeedInfo.IsValid())
 	{
-		// If the camera feed dimensions are valid, use them and the dimensions of the media to set the aspect ratio correction material parameters
-		MaterialPass->Material.SetScalarOverride(TEXT("AspectRatioCorrection_H"), MediaPlateRenderTarget->SizeX / (float)CameraFeedDimensions.X);
-		MaterialPass->Material.SetScalarOverride(TEXT("AspectRatioCorrection_V"), MediaPlateRenderTarget->SizeY / (float)CameraFeedDimensions.Y);
+		const FIntPoint CompRenderResolution = CompPtr->GetRenderResolution();
+		const FIntPoint CameraFeedDimensions = LensFilePtr->CameraFeedInfo.GetDimensions();
+
+		AspectRatioCorrectionX = CompRenderResolution.X / (float)CameraFeedDimensions.X;
+		AspectRatioCorrectionY = CompRenderResolution.Y / (float)CameraFeedDimensions.Y;
 	}
+
+	// If the camera feed dimensions are invalid, reset the correction parameters
+	MaterialPassPtr->Material.SetScalarOverride(TEXT("AspectRatioCorrection_H"), AspectRatioCorrectionX);
+	MaterialPassPtr->Material.SetScalarOverride(TEXT("AspectRatioCorrection_V"), AspectRatioCorrectionY);
 }
 
 void FCameraCalibrationStepsController::MinimizeAspectRatioError(FIntPoint& CameraFeedDimensions, float CameraAspectRatio)
@@ -490,9 +558,32 @@ UTextureRenderTarget2D* FCameraCalibrationStepsController::GetRenderTarget() con
 	return RenderTarget.Get();
 }
 
-FIntPoint FCameraCalibrationStepsController::GetCompRenderTargetSize() const
+FIntPoint FCameraCalibrationStepsController::GetCompRenderResolution() const
 {
-	return RenderTargetSize;
+	if (ACompositingElement* CompPtr = Comp.Get())
+	{
+		return CompPtr->GetRenderResolution();
+	}
+	return FIntPoint(0, 0);
+}
+
+FIntPoint FCameraCalibrationStepsController::GetCGRenderResolution() const
+{
+	if (ACompositingElement* CGLayerPtr = CGLayer.Get())
+	{
+		return CGLayerPtr->GetRenderResolution();
+	}
+	return FIntPoint(0, 0);
+}
+
+FIntPoint FCameraCalibrationStepsController::GetCameraFeedSize() const
+{
+	if (ULensFile* LensFilePtr = GetLensFile())
+	{
+		return LensFilePtr->CameraFeedInfo.GetDimensions();
+	}
+
+	return FIntPoint(0, 0);
 }
 
 void FCameraCalibrationStepsController::CreateComp()
@@ -576,6 +667,9 @@ void FCameraCalibrationStepsController::CreateComp()
 		return;
 	}
 
+	// The CG Layer's resolution will be tied to the aspect ratio of the camera's filmback, not the resolution of the top-level comp
+	CGLayer->ResolutionSource = EInheritedSourceType::Override;
+
 	// This updates the Composure panel view
 	CompElementManager->RefreshElementsList();
 
@@ -723,6 +817,28 @@ void FCameraCalibrationStepsController::CreateComp()
 	CreateOverlayPass(TEXT("Calibration Step Overlay"), ToolOverlayPass, ToolOverlayRenderTarget);
 	CreateOverlayPass(TEXT("User Selected Overlay"), UserOverlayPass, UserOverlayRenderTarget);
 
+	const FIntPoint InitialCompRenderResolution = Comp->GetRenderResolution();
+
+	if (ULensFile* LensFilePtr = LensFile.Get())
+	{
+		// Initialize the simulcam info of the LensFile to match the initial render resolution of the comp
+		const float CompAspectRatio = (InitialCompRenderResolution.Y != 0) ? InitialCompRenderResolution.X / (float)InitialCompRenderResolution.Y : 1.0f;
+		LensFilePtr->SimulcamInfo.CGLayerAspectRatio = CompAspectRatio;
+		LensFilePtr->SimulcamInfo.MediaPlateAspectRatio = CompAspectRatio;
+
+		// If the Camera Feed is not specifically overriden by the user, initialize it to the initial dimensions of the comp as well
+		// If the camera feed exactly matches the comp's resolution, then no aspect ratio correction is needed
+		if (!LensFilePtr->CameraFeedInfo.IsOverridden())
+		{
+			LensFilePtr->CameraFeedInfo.SetDimensions(InitialCompRenderResolution);
+		}
+		else
+		{
+			// If the Camera Feed has been overridden by the user, update the aspect ratio correction material parameters
+			UpdateAspectRatioCorrection();
+		}
+	}
+
 	URenderTargetCompositingOutput* RTOutput = Cast<URenderTargetCompositingOutput>(Comp->CreateNewOutputPass(
 		TEXT("SimulcamCalOutput"),
 		URenderTargetCompositingOutput::StaticClass())
@@ -751,7 +867,7 @@ void FCameraCalibrationStepsController::CreateComp()
 	RenderTarget->RenderTargetFormat = RTF_RGBA16f;
 	RenderTarget->ClearColor = FLinearColor::Black;
 	RenderTarget->bAutoGenerateMips = false;
-	RenderTarget->InitAutoFormat(RenderTargetSize.X, RenderTargetSize.Y);
+	RenderTarget->InitAutoFormat(InitialCompRenderResolution.X, InitialCompRenderResolution.Y);
 	RenderTarget->UpdateResourceImmediate(true);
 
 	// Assign the RT to the compositing output
@@ -900,10 +1016,12 @@ void FCameraCalibrationStepsController::CreateMediaPlateOutput()
 		return;
 	}
 
+	const FIntPoint InitialCompRenderResolution = GetCompRenderResolution();
+
 	MediaPlateRenderTarget->RenderTargetFormat = ETextureRenderTargetFormat::RTF_RGBA8;
 	MediaPlateRenderTarget->ClearColor = FLinearColor::Black;
 	MediaPlateRenderTarget->bAutoGenerateMips = false;
-	MediaPlateRenderTarget->InitAutoFormat(RenderTargetSize.X, RenderTargetSize.Y);
+	MediaPlateRenderTarget->InitAutoFormat(InitialCompRenderResolution.X, InitialCompRenderResolution.Y);
 	MediaPlateRenderTarget->UpdateResourceImmediate(true);
 
 	// Assign the RT to the compositing output
@@ -943,10 +1061,12 @@ void FCameraCalibrationStepsController::CreateOverlayPass(FName PassName, TWeakO
 		return;
 	}
 
+	const FIntPoint InitialCompRenderResolution = GetCompRenderResolution();
+
 	OverlayRenderTarget->RenderTargetFormat = ETextureRenderTargetFormat::RTF_RGBA8;
-	OverlayRenderTarget->ClearColor = FLinearColor::Black;
+	OverlayRenderTarget->ClearColor = FLinearColor::Transparent;
 	OverlayRenderTarget->bAutoGenerateMips = false;
-	OverlayRenderTarget->InitAutoFormat(RenderTargetSize.X, RenderTargetSize.Y);
+	OverlayRenderTarget->InitAutoFormat(InitialCompRenderResolution.X, InitialCompRenderResolution.Y);
 	OverlayRenderTarget->UpdateResourceImmediate(true);
 
 	OverlayPass->Material.SetTextureOverride(FName(TEXT("OverlayTexture")), OverlayRenderTarget.Get());
@@ -1504,8 +1624,7 @@ UTextureRenderTarget2D* FCameraCalibrationStepsController::GetMediaPlateRenderTa
 	return MediaPlateRenderTarget.Get();
 }
 
-
-bool FCameraCalibrationStepsController::CalculateNormalizedMouseClickPosition(const FGeometry& MyGeometry, const FPointerEvent& MouseEvent, FVector2D& OutPosition) const
+bool FCameraCalibrationStepsController::CalculateNormalizedMouseClickPosition(const FGeometry& MyGeometry, const FPointerEvent& MouseEvent, FVector2D& OutPosition, ESimulcamViewportPortion ViewportPortion) const
 {
 	// Reject viewports with no area
 	if (FMath::IsNearlyZero(MyGeometry.Size.X) || FMath::IsNearlyZero(MyGeometry.Size.Y))
@@ -1522,8 +1641,32 @@ bool FCameraCalibrationStepsController::CalculateNormalizedMouseClickPosition(co
 
 	const FVector2D LocalInPixels = MyGeometry.AbsoluteToLocal(MouseEvent.GetScreenSpacePosition());
 
-	const float XNormalized = LocalInPixels.X / MyGeometry.Size.X;
-	const float YNormalized = LocalInPixels.Y / MyGeometry.Size.Y;
+	float XNormalized = LocalInPixels.X / MyGeometry.Size.X;
+	float YNormalized = LocalInPixels.Y / MyGeometry.Size.Y;
+
+	if (ViewportPortion == ESimulcamViewportPortion::CameraFeed)
+	{
+		ULensFile* LensFilePtr = GetLensFile();
+		if (!LensFilePtr)
+		{
+			return false;
+		}
+
+		const FIntPoint CameraFeedDimensions = LensFilePtr->CameraFeedInfo.GetDimensions();
+		const FIntPoint CompRenderResolution = GetCompRenderResolution();
+
+		const float AspectRatioCorrectionX = CompRenderResolution.X / (float)CameraFeedDimensions.X;
+		const float AspectRatioCorrectionY = CompRenderResolution.Y / (float)CameraFeedDimensions.Y;
+
+		XNormalized = ((XNormalized - 0.5f) * AspectRatioCorrectionX) + 0.5f;
+		YNormalized = ((YNormalized - 0.5f) * AspectRatioCorrectionY) + 0.5f;
+
+		// If the scaled values for X or Y are outside the range of [0,1], then the position is invalid (not on the camera feed)
+		if (XNormalized < 0.0f || XNormalized > 1.0f || YNormalized < 0.0f || YNormalized > 1.0f)
+		{
+			return false;
+		}
+	}
 
 	// Position 0~1. Origin at top-left corner of the viewport.
 	OutPosition = FVector2D(XNormalized, YNormalized);
@@ -1531,7 +1674,7 @@ bool FCameraCalibrationStepsController::CalculateNormalizedMouseClickPosition(co
 	return true;
 }
 
-bool FCameraCalibrationStepsController::ReadMediaPixels(TArray<FColor>& Pixels, FIntPoint& Size, ETextureRenderTargetFormat& PixelFormat, FText& OutErrorMessage) const
+bool FCameraCalibrationStepsController::ReadMediaPixels(TArray<FColor>& Pixels, FIntPoint& Size, ETextureRenderTargetFormat& PixelFormat, FText& OutErrorMessage, ESimulcamViewportPortion ViewportPortion) const
 {
 	// Get the media plate texture render target 2d
 
@@ -1553,7 +1696,8 @@ bool FCameraCalibrationStepsController::ReadMediaPixels(TArray<FColor>& Pixels,
 	PixelFormat = MediaPlateRenderTarget->RenderTargetFormat;
 
 	// Read the pixels onto CPU
-	const bool bReadPixels = MediaRenderTarget->ReadPixels(Pixels);
+	TArray<FColor> MediaPixels;
+	const bool bReadPixels = MediaRenderTarget->ReadPixels(MediaPixels);
 
 	if (!bReadPixels)
 	{
@@ -1561,7 +1705,40 @@ bool FCameraCalibrationStepsController::ReadMediaPixels(TArray<FColor>& Pixels,
 		return false;
 	}
 
-	Size = MediaRenderTarget->GetSizeXY();
+	ULensFile* LensFilePtr = GetLensFile();
+	if (!LensFilePtr)
+	{
+		OutErrorMessage = LOCTEXT("InvalidLensFile", "There was no LensFile found.");
+		return false;
+	}
+
+	if ((ViewportPortion == ESimulcamViewportPortion::CameraFeed) && LensFilePtr->CameraFeedInfo.IsValid())
+	{
+		Size = LensFilePtr->CameraFeedInfo.GetDimensions();
+		Pixels.SetNumUninitialized(Size.X * Size.Y);
+
+		const FIntPoint MediaResolution = MediaRenderTarget->GetSizeXY();
+		const FIntPoint MediaCenterPoint = MediaResolution / 2;
+		const float AspectRatioCorrectionX = MediaResolution.X / (float)Size.X;
+		const float AspectRatioCorrectionY = MediaResolution.Y / (float)Size.Y;
+
+		// Only return pixels from the actual camera feed (which may be smaller than the full media render target)
+		for (int32 YCoordinate = 0; YCoordinate < Size.Y; YCoordinate++)
+		{
+			for (int32 XCoordinate = 0; XCoordinate < Size.X; XCoordinate++)
+			{
+				const int32 ScaledX = (((XCoordinate * AspectRatioCorrectionX) - MediaCenterPoint.X) / AspectRatioCorrectionX) + MediaCenterPoint.X;
+				const int32 ScaledY = (((YCoordinate * AspectRatioCorrectionY) - MediaCenterPoint.Y) / AspectRatioCorrectionY) + MediaCenterPoint.Y;
+
+				Pixels[YCoordinate * Size.X + XCoordinate] = MediaPixels[ScaledY * MediaResolution.X + ScaledX];
+			}
+		}
+	}
+	else
+	{
+		Pixels = MoveTemp(MediaPixels);
+		Size = MediaRenderTarget->GetSizeXY();
+	}
 
 	check(Pixels.Num() == Size.X * Size.Y);
 
diff --git a/Engine/Plugins/VirtualProduction/CameraCalibration/Source/CameraCalibrationEditor/Private/AssetEditor/CameraCalibrationStepsController.h b/Engine/Plugins/VirtualProduction/CameraCalibration/Source/CameraCalibrationEditor/Private/AssetEditor/CameraCalibrationStepsController.h
--- a/Engine/Plugins/VirtualProduction/CameraCalibration/Source/CameraCalibrationEditor/Private/AssetEditor/CameraCalibrationStepsController.h
+++ b/Engine/Plugins/VirtualProduction/CameraCalibration/Source/CameraCalibrationEditor/Private/AssetEditor/CameraCalibrationStepsController.h
@@ -34,6 +34,13 @@ enum class EOverlayPassType : uint8
 	UserOverlay = 1
 };
 
+/** Enumeration specifying the portion of the viewport to consider when performing operations such as normalizing the mouse position or reading pixels from the media render target */
+enum class ESimulcamViewportPortion : uint8
+{
+	FullViewport = 0,
+	CameraFeed = 1
+};
+
 /**
  * Controller for SCameraCalibrationSteps, where the calibration steps are hosted in.
  */
@@ -59,8 +66,14 @@ public:
 	/** Returns the render target of the Media Plate */
 	UTextureRenderTarget2D* GetMediaPlateRenderTarget() const;
 
-	/** Returns the size of the render target used by the Comp */
-	FIntPoint GetCompRenderTargetSize() const;
+	/** Returns the output resolution of the top-level Comp */
+	FIntPoint GetCompRenderResolution() const;
+
+	/** Returns the output resolution of the CG Layer */
+	FIntPoint GetCGRenderResolution() const;
+
+	/** Returns the size of the camera feed (which may be smaller than the top-level Comp and CG Layer) */
+	FIntPoint GetCameraFeedSize() const;
 
 	/** Creates a way to read the media plate pixels for processing by any calibration step */
 	void CreateMediaPlateOutput();
@@ -117,13 +130,13 @@ public:
 	void SelectStep(const FName& Name);
 
 	/** Calculates the normalized (0~1) coordinates in the simulcam viewport of the given mouse click */
-	bool CalculateNormalizedMouseClickPosition(const FGeometry& MyGeometry, const FPointerEvent& MouseEvent, FVector2D& OutPosition) const;
+	bool CalculateNormalizedMouseClickPosition(const FGeometry& MyGeometry, const FPointerEvent& MouseEvent, FVector2D& OutPosition, ESimulcamViewportPortion ViewportPortion = ESimulcamViewportPortion::FullViewport) const;
 
 	/** Finds the world being used by the tool for finding and spawning objects */
 	UWorld* GetWorld() const;
 
 	/** Reads the pixels in the media plate */
-	bool ReadMediaPixels(TArray<FColor>& Pixels, FIntPoint& Size, ETextureRenderTargetFormat& PixelFormat, FText& OutErrorMessage) const;
+	bool ReadMediaPixels(TArray<FColor>& Pixels, FIntPoint& Size, ETextureRenderTargetFormat& PixelFormat, FText& OutErrorMessage, ESimulcamViewportPortion ViewportPortion = ESimulcamViewportPortion::FullViewport) const;
 
 	/** Returns true if the overlay transform pass is currently enabled */
 	bool IsOverlayEnabled(EOverlayPassType OverlayPass = EOverlayPassType::ToolOverlay) const;
@@ -138,10 +151,10 @@ public:
 	void RefreshOverlay(EOverlayPassType OverlayPass = EOverlayPassType::ToolOverlay);
 
 	/** Use the input mouse position (representing any corner of the camera feed) and the dimensions of the media source to calculate the new size of the camera feed */
-	void UpdateCameraFeedInfo(FVector2D MousePosition);
+	void SetCameraFeedDimensionsFromMousePosition(FVector2D MousePosition);
 
-	/** Update the material parameters that control the CG layer's aspect ratio correction based on the camera feed (if it is valid) or on the aspect ratio of the camera actor */
-	void UpdateAspectRatioCorrection(FIntPoint CameraFeedDimensions);
+	/** Set the camera feed dimensions info of the LensFile being edited */
+	void SetCameraFeedDimensions(FIntPoint Dimensions, bool bMarkAsOverridden);
 
 public:
 
@@ -237,14 +250,20 @@ private:
 	/** Iteratively attempt to resize the camera feed until its aspect ratio matches the input camera aspect ratio */
 	void MinimizeAspectRatioError(FIntPoint& CameraFeedDimensions, float CameraAspectRatio);
 
+	/** Resizes the resolution of the top-level comp match the resolution of the media source */
+	bool UpdateCompResolution();
+
+	/** Resizes the resolution of the CG layer to match the aspect resolution of the source camera and fit within the comp */
+	void UpdateCGResolution(bool bCompResized);
+
+	/** Update the aspect ratio correction material parameters based on the ratio of the camera feed dimensions to the comp dimensions */
+	void UpdateAspectRatioCorrection();
+
 private:
 
 	/** Pointer to the camera calibration toolkit */
 	TWeakPtr<FCameraCalibrationToolkit> CameraCalibrationToolkit;
 
-	/** Size to use when creating the render targets for the comp and media output */
-	FIntPoint RenderTargetSize;
-
 	/** Array of the calibration steps that this controller is managing */
 	TArray<TStrongObjectPtr<UCameraCalibrationStep>> CalibrationSteps;
 
diff --git a/Engine/Plugins/VirtualProduction/CameraCalibration/Source/CameraCalibrationEditor/Private/AssetEditor/LensInfoStep.cpp b/Engine/Plugins/VirtualProduction/CameraCalibration/Source/CameraCalibrationEditor/Private/AssetEditor/LensInfoStep.cpp
--- a/Engine/Plugins/VirtualProduction/CameraCalibration/Source/CameraCalibrationEditor/Private/AssetEditor/LensInfoStep.cpp
+++ b/Engine/Plugins/VirtualProduction/CameraCalibration/Source/CameraCalibrationEditor/Private/AssetEditor/LensInfoStep.cpp
@@ -89,6 +89,16 @@ TSharedRef<SWidget> ULensInfoStep::BuildUI()
 			SNew(SHorizontalBox)
 			+ SHorizontalBox::Slot()
 			[CameraFeedInfoStructureDetailsView->GetWidget().ToSharedRef() ]
+
+			+ SHorizontalBox::Slot()
+			.AutoWidth()
+			.Padding(2.0f)
+			.VAlign(VAlign_Center)
+			[
+				SNew(SResetToDefaultMenu)
+				.OnResetToDefault(FSimpleDelegate::CreateUObject(this, &ULensInfoStep::ResetCameraFeedInfoToDefault))
+				.DiffersFromDefault(TAttribute<bool>::Create(TAttribute<bool>::FGetter::CreateUObject(this, &ULensInfoStep::CameraFeedInfoDiffersFromDefault)))
+			]
 		]
 
 		+ SVerticalBox::Slot()
@@ -210,7 +220,7 @@ bool ULensInfoStep::OnViewportClicked(const FGeometry& MyGeometry, const FPointe
 		if (TSharedPtr<FCameraCalibrationStepsController> StepsController = CameraCalibrationStepsController.Pin())
 		{
 			const FVector2D LocalInPixels = MyGeometry.AbsoluteToLocal(MouseEvent.GetScreenSpacePosition());
-			StepsController->UpdateCameraFeedInfo(LocalInPixels);
+			StepsController->SetCameraFeedDimensionsFromMousePosition(LocalInPixels);
 			return true;
 		}
 	}
@@ -227,13 +237,38 @@ void ULensInfoStep::OnCameraFeedInfoChanged(const FPropertyChangedEvent& Propert
 		{
 			if (ULensFile* LensFile = StepsController->GetLensFile())
 			{
-				LensFile->CameraFeedInfo.SetDimensions(LensFile->CameraFeedInfo.GetDimensions());
-				StepsController->UpdateAspectRatioCorrection(LensFile->CameraFeedInfo.GetDimensions());
+				FIntPoint CameraFeedDimensions = LensFile->CameraFeedInfo.GetDimensions();
+				const FIntPoint CompRenderResolution = StepsController->GetCompRenderResolution();
+				CameraFeedDimensions.X = FMath::Clamp(CameraFeedDimensions.X, 0, CompRenderResolution.X);
+				CameraFeedDimensions.Y = FMath::Clamp(CameraFeedDimensions.Y, 0, CompRenderResolution.Y);
+
+				StepsController->SetCameraFeedDimensions(CameraFeedDimensions, true);
 			}
 		}
 	}
 }
 
+void ULensInfoStep::ResetCameraFeedInfoToDefault()
+{
+	if (TSharedPtr<FCameraCalibrationStepsController> StepsController = CameraCalibrationStepsController.Pin())
+	{
+		const FIntPoint DefaultCameraFeedDimensions = StepsController->GetCGRenderResolution();
+		StepsController->SetCameraFeedDimensions(DefaultCameraFeedDimensions, false);
+	}
+}
+
+bool ULensInfoStep::CameraFeedInfoDiffersFromDefault() const
+{
+	if (TSharedPtr<FCameraCalibrationStepsController> StepsController = CameraCalibrationStepsController.Pin())
+	{
+		if (ULensFile* LensFile = StepsController->GetLensFile())
+		{
+			return LensFile->CameraFeedInfo.IsOverridden();
+		}
+	}
+	return false;
+}
+
 EVisibility ULensInfoStep::HandleAspectRatioWarningVisibility() const
 {
 	if (TSharedPtr<FCameraCalibrationStepsController> StepsController = CameraCalibrationStepsController.Pin())
diff --git a/Engine/Plugins/VirtualProduction/CameraCalibration/Source/CameraCalibrationEditor/Private/AssetEditor/LensInfoStep.h b/Engine/Plugins/VirtualProduction/CameraCalibration/Source/CameraCalibrationEditor/Private/AssetEditor/LensInfoStep.h
--- a/Engine/Plugins/VirtualProduction/CameraCalibration/Source/CameraCalibrationEditor/Private/AssetEditor/LensInfoStep.h
+++ b/Engine/Plugins/VirtualProduction/CameraCalibration/Source/CameraCalibrationEditor/Private/AssetEditor/LensInfoStep.h
@@ -44,6 +44,13 @@ public:
 	/** Returns true if our LensInfo differs from the original info when the tool was opened */
 	bool DiffersFromDefault() const;
 
+	/** Reset the camera feed dimensions to their default value (recomputed by the Steps Controller) */
+	void ResetCameraFeedInfoToDefault();
+
+	/** Returns true if the camera feed info has been overriden by the user */
+	bool CameraFeedInfoDiffersFromDefault() const;
+
+
 private:
 	/** Determines the visibility of the aspect ratio warning by comparing the camera feed aspect ratio with that of the camera actor */
 	EVisibility HandleAspectRatioWarningVisibility() const;
diff --git a/Engine/Plugins/VirtualProduction/CameraCalibration/Source/CameraCalibrationEditor/Private/Calibrators/CameraLensDistortionAlgoCheckerboard.cpp b/Engine/Plugins/VirtualProduction/CameraCalibration/Source/CameraCalibrationEditor/Private/Calibrators/CameraLensDistortionAlgoCheckerboard.cpp
--- a/Engine/Plugins/VirtualProduction/CameraCalibration/Source/CameraCalibrationEditor/Private/Calibrators/CameraLensDistortionAlgoCheckerboard.cpp
+++ b/Engine/Plugins/VirtualProduction/CameraCalibration/Source/CameraCalibrationEditor/Private/Calibrators/CameraLensDistortionAlgoCheckerboard.cpp
@@ -120,7 +120,7 @@ void UCameraLensDistortionAlgoCheckerboard::Initialize(ULensDistortionTool* InTo
 		return;
 	}
 
-	FIntPoint Size = StepsController->GetCompRenderTargetSize();
+	FIntPoint Size = StepsController->GetCompRenderResolution();
 	CvCoverage = cv::Mat(cv::Size(Size.X, Size.Y), CV_8UC4);
 	CoverageTexture = FOpenCVHelper::TextureFromCvMat(CvCoverage, CoverageTexture);
 #endif
@@ -247,7 +247,7 @@ bool UCameraLensDistortionAlgoCheckerboard::AddCalibrationRow(FText& OutErrorMes
 	FIntPoint Size;
 	ETextureRenderTargetFormat PixelFormat;
 
-	if (!StepsController->ReadMediaPixels(Pixels, Size, PixelFormat, OutErrorMessage))
+	if (!StepsController->ReadMediaPixels(Pixels, Size, PixelFormat, OutErrorMessage, ESimulcamViewportPortion::CameraFeed))
 	{
 		return false;
 	}
@@ -332,12 +332,6 @@ bool UCameraLensDistortionAlgoCheckerboard::AddCalibrationRow(FText& OutErrorMes
 		FImageView ImageView = FImageView(CvFrame.data, CvFrame.cols, CvFrame.rows, ERawImageFormat::BGRA8);
 		Row->ImageView = ImageView;
 
-		// Update the coverage overlay image with information from the newly added row
-		cv::drawChessboardCorners(CvCoverage, CheckerboardSize, Corners, true);
-
-		CoverageTexture = FOpenCVHelper::TextureFromCvMat(CvCoverage, CoverageTexture);
-		StepsController->RefreshOverlay();
-
 		// Show the detection to the user
 		if (bShouldShowDetectionWindow)
 		{
@@ -350,6 +344,31 @@ bool UCameraLensDistortionAlgoCheckerboard::AddCalibrationRow(FText& OutErrorMes
 				LOCTEXT("CheckerboardDetection", "Checkerboard Detection")
 			);
 		}
+
+		// The corners were found using only the camera feed, but the coverage overlay will be drawn over the entire comp
+		// Before updating the overlay texture, shift all the camera feed corner pixels to represent the correct position in the overlay texture
+		const FIntPoint CompSize = StepsController->GetCompRenderResolution();
+		const FIntPoint CameraFeedSize = StepsController->GetCameraFeedSize();
+
+		const FIntPoint TopLeftCorner = (CompSize - CameraFeedSize) / 2;
+
+		for (cv::Point2f& Corner : Corners)
+		{
+			Corner.x += TopLeftCorner.X;
+			Corner.y += TopLeftCorner.Y;
+		}
+
+		// Refresh the coverage overlay if the comp resolution has changed since the last image was taken
+		if ((CvCoverage.cols != CompSize.X) || (CvCoverage.rows != CompSize.Y))
+		{
+			RefreshCoverage();
+		}
+
+		// Update the coverage overlay image with information from the newly added row
+		cv::drawChessboardCorners(CvCoverage, CheckerboardSize, Corners, true);
+
+		CoverageTexture = FOpenCVHelper::TextureFromCvMat(CvCoverage, CoverageTexture);
+		StepsController->RefreshOverlay();
 	}
 
 	// Create thumbnail and add it to the row
@@ -959,7 +978,7 @@ void UCameraLensDistortionAlgoCheckerboard::RefreshCoverage()
 		return;
 	}
 
-	FIntPoint Size = StepsController->GetCompRenderTargetSize();
+	FIntPoint Size = StepsController->GetCompRenderResolution();
 
 	CvCoverage.release();
 	CvCoverage = cv::Mat(cv::Size(Size.X, Size.Y), CV_8UC4);
@@ -980,6 +999,16 @@ void UCameraLensDistortionAlgoCheckerboard::RefreshCoverage()
 
 	CoverageTexture = FOpenCVHelper::TextureFromCvMat(CvCoverage, CoverageTexture);
 
+	// The coverage texture may have changed as a result of a change in size or pixel format.
+	// Therefore, the material parameter should be updated to ensure it is up to date.
+	if (CoverageTexture)
+	{
+		if (UMaterialInstanceDynamic* OverlayMID = Tool->GetOverlayMID())
+		{
+			OverlayMID->SetTextureParameterValue(FName(TEXT("CoverageTexture")), CoverageTexture);
+		}
+	}
+
 	StepsController->RefreshOverlay();
 #endif
 }
diff --git a/Engine/Plugins/VirtualProduction/CameraCalibration/Source/CameraCalibrationEditor/Private/Calibrators/CameraLensDistortionAlgoPoints.cpp b/Engine/Plugins/VirtualProduction/CameraCalibration/Source/CameraCalibrationEditor/Private/Calibrators/CameraLensDistortionAlgoPoints.cpp
--- a/Engine/Plugins/VirtualProduction/CameraCalibration/Source/CameraCalibrationEditor/Private/Calibrators/CameraLensDistortionAlgoPoints.cpp
+++ b/Engine/Plugins/VirtualProduction/CameraCalibration/Source/CameraCalibrationEditor/Private/Calibrators/CameraLensDistortionAlgoPoints.cpp
@@ -254,7 +254,7 @@ bool UCameraLensDistortionAlgoPoints::OnViewportClicked(const FGeometry& MyGeome
 	}
 
 	// Get the mouse click 2d position
-	if (!StepsController->CalculateNormalizedMouseClickPosition(MyGeometry, MouseEvent, LastCalibratorPoint.Point2d))
+	if (!StepsController->CalculateNormalizedMouseClickPosition(MyGeometry, MouseEvent, LastCalibratorPoint.Point2d, ESimulcamViewportPortion::CameraFeed))
 	{
 		return true;
 	}
@@ -383,7 +383,7 @@ bool UCameraLensDistortionAlgoPoints::GetLensDistortion(
 		return false;
 	}
 
-	const FIntPoint ImageSize = StepsController->GetCompRenderTargetSize();
+	const FIntPoint ImageSize = StepsController->GetCompRenderResolution();
 
 	// Gather 3D points and 2D points from each group of rows
 	std::vector<std::vector<cv::Point3f>> Samples3d;
diff --git a/Engine/Plugins/VirtualProduction/CameraCalibration/Source/CameraCalibrationEditor/Private/Calibrators/CameraNodalOffsetAlgoAruco.cpp b/Engine/Plugins/VirtualProduction/CameraCalibration/Source/CameraCalibrationEditor/Private/Calibrators/CameraNodalOffsetAlgoAruco.cpp
--- a/Engine/Plugins/VirtualProduction/CameraCalibration/Source/CameraCalibrationEditor/Private/Calibrators/CameraNodalOffsetAlgoAruco.cpp
+++ b/Engine/Plugins/VirtualProduction/CameraCalibration/Source/CameraCalibrationEditor/Private/Calibrators/CameraNodalOffsetAlgoAruco.cpp
@@ -204,7 +204,7 @@ bool UCameraNodalOffsetAlgoAruco::PopulatePoints(FText& OutErrorMessage)
 	FIntPoint Size;
 	ETextureRenderTargetFormat PixelFormat;
 
-	if (!StepsController->ReadMediaPixels(Pixels, Size, PixelFormat, OutErrorMessage))
+	if (!StepsController->ReadMediaPixels(Pixels, Size, PixelFormat, OutErrorMessage, ESimulcamViewportPortion::CameraFeed))
 	{
 		return false;
 	}
diff --git a/Engine/Plugins/VirtualProduction/CameraCalibration/Source/CameraCalibrationEditor/Private/Calibrators/CameraNodalOffsetAlgoCheckerboard.cpp b/Engine/Plugins/VirtualProduction/CameraCalibration/Source/CameraCalibrationEditor/Private/Calibrators/CameraNodalOffsetAlgoCheckerboard.cpp
--- a/Engine/Plugins/VirtualProduction/CameraCalibration/Source/CameraCalibrationEditor/Private/Calibrators/CameraNodalOffsetAlgoCheckerboard.cpp
+++ b/Engine/Plugins/VirtualProduction/CameraCalibration/Source/CameraCalibrationEditor/Private/Calibrators/CameraNodalOffsetAlgoCheckerboard.cpp
@@ -132,7 +132,7 @@ bool UCameraNodalOffsetAlgoCheckerboard::PopulatePoints(FText& OutErrorMessage)
 	FIntPoint Size;
 	ETextureRenderTargetFormat PixelFormat;
 
-	if (!StepsController->ReadMediaPixels(Pixels, Size, PixelFormat, OutErrorMessage))
+	if (!StepsController->ReadMediaPixels(Pixels, Size, PixelFormat, OutErrorMessage, ESimulcamViewportPortion::CameraFeed))
 	{
 		return false;
 	}
diff --git a/Engine/Plugins/VirtualProduction/CameraCalibration/Source/CameraCalibrationEditor/Private/Calibrators/CameraNodalOffsetAlgoOpticalAxis.cpp b/Engine/Plugins/VirtualProduction/CameraCalibration/Source/CameraCalibrationEditor/Private/Calibrators/CameraNodalOffsetAlgoOpticalAxis.cpp
--- a/Engine/Plugins/VirtualProduction/CameraCalibration/Source/CameraCalibrationEditor/Private/Calibrators/CameraNodalOffsetAlgoOpticalAxis.cpp
+++ b/Engine/Plugins/VirtualProduction/CameraCalibration/Source/CameraCalibrationEditor/Private/Calibrators/CameraNodalOffsetAlgoOpticalAxis.cpp
@@ -46,7 +46,7 @@ void UCameraNodalOffsetAlgoOpticalAxis::Initialize(UNodalOffsetTool* InNodalOffs
 	}
 
 	// Create the render target for the crosshair overlay
-	const FIntPoint OverlaySize = StepsController->GetCompRenderTargetSize();
+	const FIntPoint OverlaySize = StepsController->GetCompRenderResolution();
 
 	OverlayTexture = NewObject<UTextureRenderTarget2D>(GetTransientPackage(), MakeUniqueObjectName(GetTransientPackage(), UTextureRenderTarget2D::StaticClass()));
 
diff --git a/Engine/Plugins/VirtualProduction/CameraCalibration/Source/CameraCalibrationEditor/Private/Calibrators/CameraNodalOffsetAlgoPoints.cpp b/Engine/Plugins/VirtualProduction/CameraCalibration/Source/CameraCalibrationEditor/Private/Calibrators/CameraNodalOffsetAlgoPoints.cpp
--- a/Engine/Plugins/VirtualProduction/CameraCalibration/Source/CameraCalibrationEditor/Private/Calibrators/CameraNodalOffsetAlgoPoints.cpp
+++ b/Engine/Plugins/VirtualProduction/CameraCalibration/Source/CameraCalibrationEditor/Private/Calibrators/CameraNodalOffsetAlgoPoints.cpp
@@ -752,7 +752,7 @@ bool UCameraNodalOffsetAlgoPoints::OnViewportClicked(const FGeometry& MyGeometry
 	Row->CalibratorPointData = LastCalibratorPoint;
 
 	// Get the mouse click 2d position
-	if (!StepsController->CalculateNormalizedMouseClickPosition(MyGeometry, MouseEvent, Row->Point2D))
+	if (!StepsController->CalculateNormalizedMouseClickPosition(MyGeometry, MouseEvent, Row->Point2D, ESimulcamViewportPortion::CameraFeed))
 	{
 		return true;
 	}
diff --git a/Engine/Plugins/VirtualProduction/CameraCalibrationCore/Source/CameraCalibrationCore/Public/LensData.h b/Engine/Plugins/VirtualProduction/CameraCalibrationCore/Source/CameraCalibrationCore/Public/LensData.h
--- a/Engine/Plugins/VirtualProduction/CameraCalibrationCore/Source/CameraCalibrationCore/Public/LensData.h
+++ b/Engine/Plugins/VirtualProduction/CameraCalibrationCore/Source/CameraCalibrationCore/Public/LensData.h
@@ -70,14 +70,15 @@ struct CAMERACALIBRATIONCORE_API FCameraFeedInfo
 	GENERATED_BODY()
 
 public:
-	FIntPoint GetDimensions()
+	FIntPoint GetDimensions() const
 	{
 		return Dimensions;
 	}
 
-	void SetDimensions(FIntPoint InDimensions)
+	void SetDimensions(FIntPoint InDimensions, bool bMarkAsOverridden = false)
 	{
 		Dimensions = InDimensions;
+		bIsOverridden = bMarkAsOverridden;
 
 		// If one of the dimensions is set to 0, the description of the feed is invalid
 		if (Dimensions.X == 0 || Dimensions.Y == 0)
@@ -92,7 +93,7 @@ public:
 		}
 	}
 
-	float GetAspectRatio()
+	float GetAspectRatio() const
 	{
 		return AspectRatio;
 	}
@@ -104,11 +105,16 @@ public:
 		bIsValid = false;
 	}
 
-	bool IsValid()
+	bool IsValid() const
 	{
 		return bIsValid;
 	}
 
+	bool IsOverridden() const
+	{
+		return bIsOverridden;
+	}
+
 private:
 	/** Dimensions of the Camera Feed */
 	UPROPERTY(EditAnywhere, Category = "Camera Feed Info")
@@ -120,6 +126,9 @@ private:
 
 	/** Whether the dimensions of the feed are valid (false if either X or Y is 0) */
 	bool bIsValid = false;
+
+	/** Whether the dimensions of the feed were overridden by the user */
+	bool bIsOverridden = false;
 };
 
 /**
