Fix PVS Static Analysis Warning

diff --git a/Engine/Plugins/VirtualProduction/CameraCalibration/Source/CameraCalibrationEditor/Private/Calibrators/CameraLensDistortionAlgoCheckerboard.cpp b/Engine/Plugins/VirtualProduction/CameraCalibration/Source/CameraCalibrationEditor/Private/Calibrators/CameraLensDistortionAlgoCheckerboard.cpp
--- a/Engine/Plugins/VirtualProduction/CameraCalibration/Source/CameraCalibrationEditor/Private/Calibrators/CameraLensDistortionAlgoCheckerboard.cpp
+++ b/Engine/Plugins/VirtualProduction/CameraCalibration/Source/CameraCalibrationEditor/Private/Calibrators/CameraLensDistortionAlgoCheckerboard.cpp
@@ -304,7 +304,7 @@ bool UCameraLensDistortionAlgoCheckerboard::AddCalibrationRow(FText& OutErrorMes
 			cv::CALIB_CB_ADAPTIVE_THRESH | cv::CALIB_CB_NORMALIZE_IMAGE
 		);
 
-		if (!bCornersFound)
+		if (!bCornersFound || Corners.empty())
 		{
 			OutErrorMessage = FText::FromString(FString::Printf(TEXT(
 				"Could not identify the expected checkerboard points of interest. "
@@ -320,12 +320,9 @@ bool UCameraLensDistortionAlgoCheckerboard::AddCalibrationRow(FText& OutErrorMes
 		cv::TermCriteria Criteria(cv::TermCriteria::Type::EPS | cv::TermCriteria::Type::COUNT, 30, 0.001);
 		cv::cornerSubPix(CvGray, Corners, cv::Size(11, 11), cv::Size(-1, -1), Criteria);
 
-		if (!Corners.empty())
+		for (const cv::Point2f& Corner : Corners)
 		{
-			for (cv::Point2f& Corner : Corners)
-			{
-				Row->Points2d.Add(FVector2D(Corner.x, Corner.y));
-			}
+			Row->Points2d.Add(FVector2D(Corner.x, Corner.y));
 		}
 
 		// Save an image view of the captured frame with the corners overlaid on it (for exporting)
