CameraCalibration: Converted asset action to new asset definition.
#rnx

diff --git a/Engine/Plugins/VirtualProduction/CameraCalibration/Source/CameraCalibrationEditor/CameraCalibrationEditor.Build.cs b/Engine/Plugins/VirtualProduction/CameraCalibration/Source/CameraCalibrationEditor/CameraCalibrationEditor.Build.cs
--- a/Engine/Plugins/VirtualProduction/CameraCalibration/Source/CameraCalibrationEditor/CameraCalibrationEditor.Build.cs
+++ b/Engine/Plugins/VirtualProduction/CameraCalibration/Source/CameraCalibrationEditor/CameraCalibrationEditor.Build.cs
@@ -11,6 +11,7 @@ namespace UnrealBuildTool.Rules
 				{
 					"ApplicationCore",
 					"AppFramework",
+					"AssetDefinition",
 					"AssetRegistry",
 					"AssetTools",
 					"CameraCalibrationCore",
diff --git a/Engine/Plugins/VirtualProduction/CameraCalibration/Source/CameraCalibrationEditor/Private/AssetDefinition_LensFile.cpp b/Engine/Plugins/VirtualProduction/CameraCalibration/Source/CameraCalibrationEditor/Private/AssetDefinition_LensFile.cpp
new file mode 100644
--- /dev/null
+++ b/Engine/Plugins/VirtualProduction/CameraCalibration/Source/CameraCalibrationEditor/Private/AssetDefinition_LensFile.cpp
@@ -0,0 +1,36 @@
+// Copyright Epic Games, Inc. All Rights Reserved.
+
+#include "AssetDefinition_LensFile.h"
+
+#include "AssetEditor/CameraCalibrationToolkit.h"
+#include "LensFile.h"
+
+#define LOCTEXT_NAMESPACE "AssetDefinition_LensFile"
+
+TSoftClassPtr<UObject> UAssetDefinition_LensFile::GetAssetClass() const
+{
+	return ULensFile::StaticClass();
+}
+
+FText UAssetDefinition_LensFile::GetAssetDisplayName() const
+{
+	return LOCTEXT("AssetDisplayName", "Lens File");
+}
+
+FLinearColor UAssetDefinition_LensFile::GetAssetColor() const
+{
+	return FLinearColor(0.1f, 1.0f, 0.1f);
+}
+
+EAssetCommandResult UAssetDefinition_LensFile::OpenAssets(const FAssetOpenArgs& OpenArgs) const
+{
+	const TArray<ULensFile*> ObjectsToOpen = OpenArgs.LoadObjects<ULensFile>();
+	for (ULensFile* Object : ObjectsToOpen)
+	{
+		FCameraCalibrationToolkit::CreateEditor(OpenArgs.GetToolkitMode(), OpenArgs.ToolkitHost, Object);
+	}
+
+	return EAssetCommandResult::Handled;
+}
+
+#undef LOCTEXT_NAMESPACE
diff --git a/Engine/Plugins/VirtualProduction/CameraCalibration/Source/CameraCalibrationEditor/Private/AssetDefinition_LensFile.h b/Engine/Plugins/VirtualProduction/CameraCalibration/Source/CameraCalibrationEditor/Private/AssetDefinition_LensFile.h
new file mode 100644
--- /dev/null
+++ b/Engine/Plugins/VirtualProduction/CameraCalibration/Source/CameraCalibrationEditor/Private/AssetDefinition_LensFile.h
@@ -0,0 +1,29 @@
+// Copyright Epic Games, Inc. All Rights Reserved.
+
+#pragma once
+
+#include "AssetDefinitionDefault.h"
+
+#include "AssetDefinition_LensFile.generated.h"
+
+UCLASS()
+class UAssetDefinition_LensFile : public UAssetDefinitionDefault
+{
+
+	GENERATED_BODY()
+
+public:
+	// Begin UAssetDefinition interface
+	virtual TConstArrayView<FAssetCategoryPath> GetAssetCategories() const override
+	{
+		static const auto Categories = { EAssetCategoryPaths::Misc };
+		return Categories;
+	}
+	virtual bool CanImport() const override { return true; }
+
+	virtual FText GetAssetDisplayName() const override;
+	virtual TSoftClassPtr<UObject> GetAssetClass() const override;
+	virtual FLinearColor GetAssetColor() const override;
+	virtual EAssetCommandResult OpenAssets(const FAssetOpenArgs& OpenArgs) const override;
+	// End UAssetDefinition interface
+};
diff --git a/Engine/Plugins/VirtualProduction/CameraCalibration/Source/CameraCalibrationEditor/Private/AssetTypeActions/AssetTypeActions_LensFile.cpp b/Engine/Plugins/VirtualProduction/CameraCalibration/Source/CameraCalibrationEditor/Private/AssetTypeActions/AssetTypeActions_LensFile.cpp
deleted file mode 100644
diff --git a/Engine/Plugins/VirtualProduction/CameraCalibration/Source/CameraCalibrationEditor/Private/AssetTypeActions/AssetTypeActions_LensFile.h b/Engine/Plugins/VirtualProduction/CameraCalibration/Source/CameraCalibrationEditor/Private/AssetTypeActions/AssetTypeActions_LensFile.h
deleted file mode 100644
diff --git a/Engine/Plugins/VirtualProduction/CameraCalibration/Source/CameraCalibrationEditor/Private/CameraCalibrationEditorModule.cpp b/Engine/Plugins/VirtualProduction/CameraCalibration/Source/CameraCalibrationEditor/Private/CameraCalibrationEditorModule.cpp
--- a/Engine/Plugins/VirtualProduction/CameraCalibration/Source/CameraCalibrationEditor/Private/CameraCalibrationEditorModule.cpp
+++ b/Engine/Plugins/VirtualProduction/CameraCalibration/Source/CameraCalibrationEditor/Private/CameraCalibrationEditorModule.cpp
@@ -8,12 +8,10 @@
 #include "AssetEditor/Curves/LensDataCurveModel.h"
 #include "AssetEditor/SCameraCalibrationCurveEditorView.h"
 #include "AssetToolsModule.h"
-#include "AssetTypeActions/AssetTypeActions_LensFile.h"
 #include "CameraCalibrationEditorLog.h"
 #include "CameraCalibrationSubsystem.h"
 #include "CameraCalibrationTypes.h"
 #include "IAssetTools.h"
-#include "IAssetTypeActions.h"
 #include "ICurveEditorModule.h"
 #include "IPlacementModeModule.h"
 #include "LensFile.h"
@@ -35,17 +33,6 @@ void FCameraCalibrationEditorModule::StartupModule()
 	FCameraCalibrationCommands::Register();
 	FCameraCalibrationEditorStyle::Get();
 
-	// Register AssetTypeActions
-	auto RegisterAssetTypeAction = [this](const TSharedRef<IAssetTypeActions>& InAssetTypeAction)
-	{
-		IAssetTools& AssetTools = FModuleManager::LoadModuleChecked<FAssetToolsModule>("AssetTools").Get();
-		RegisteredAssetTypeActions.Add(InAssetTypeAction);
-		AssetTools.RegisterAssetTypeActions(InAssetTypeAction);
-	};
-
-	// register asset type actions
-	RegisterAssetTypeAction(MakeShared<FAssetTypeActions_LensFile>());
-
 	{
 		const IWorkspaceMenuStructure& MenuStructure = WorkspaceMenu::GetMenuStructure();
 		TSharedRef<FWorkspaceItem> BrowserGroup = MenuStructure.GetDeveloperToolsMiscCategory()->GetParent()->AddGroup(
@@ -203,19 +190,6 @@ void FCameraCalibrationEditorModule::ShutdownModule()
 
 		FLevelEditorModule* LevelEditorModule = FModuleManager::GetModulePtr<FLevelEditorModule>("LevelEditor");
 
-		// Unregister AssetTypeActions
-		FAssetToolsModule* AssetToolsModule = FModuleManager::GetModulePtr<FAssetToolsModule>("AssetTools");
-
-		if (AssetToolsModule != nullptr)
-		{
-			IAssetTools& AssetTools = AssetToolsModule->Get();
-
-			for (TSharedRef<IAssetTypeActions> Action : RegisteredAssetTypeActions)
-			{
-				AssetTools.UnregisterAssetTypeActions(Action);
-			}
-		}
-
 		FCameraCalibrationCommands::Unregister();
 
 		UnregisterPlacementModeItems();
