CameraCalibration: Fix Composure CG Layer to use correct lens component when multiple CineCameras are in the scene
#rnx

diff --git a/Engine/Plugins/VirtualProduction/CameraCalibration/Source/CameraCalibrationEditor/Private/AssetEditor/CameraCalibrationStepsController.cpp b/Engine/Plugins/VirtualProduction/CameraCalibration/Source/CameraCalibrationEditor/Private/AssetEditor/CameraCalibrationStepsController.cpp
--- a/Engine/Plugins/VirtualProduction/CameraCalibration/Source/CameraCalibrationEditor/Private/AssetEditor/CameraCalibrationStepsController.cpp
+++ b/Engine/Plugins/VirtualProduction/CameraCalibration/Source/CameraCalibrationEditor/Private/AssetEditor/CameraCalibrationStepsController.cpp
@@ -1129,6 +1129,15 @@ void FCameraCalibrationStepsController::EnableDistortionInCG()
 			continue;
 		}
 
+		// Set the LensComponent used by the CG layer to match the LensComponent in use by our target camera
+		if (ACameraActor* const TargetCamera = Camera.Get())
+		{
+			if (ULensComponent* const TargetLens = FindLensComponentOnCamera(TargetCamera))
+			{
+				CaptureBase->SetLens(TargetLens);
+			}
+		}
+
 		// Enable distortion on the CG compositing layer
 		CaptureBase->SetApplyDistortion(true);
 	}
