CameraCalibration: Anamorphic Distortion Solver

diff --git a/Engine/Plugins/VirtualProduction/CameraCalibration/Source/CameraCalibrationEditor/Private/Calibrators/CameraCalibrationSolver.cpp b/Engine/Plugins/VirtualProduction/CameraCalibration/Source/CameraCalibrationEditor/Private/Calibrators/CameraCalibrationSolver.cpp
new file mode 100644
--- /dev/null
+++ b/Engine/Plugins/VirtualProduction/CameraCalibration/Source/CameraCalibrationEditor/Private/Calibrators/CameraCalibrationSolver.cpp
@@ -0,0 +1,1427 @@
+// Copyright Epic Games, Inc. All Rights Reserved.
+
+#include "CameraCalibrationSolver.h"
+
+#include "HAL/IConsoleManager.h"
+#include "Models/AnamorphicLensModel.h"
+#include "Models/SphericalLensModel.h"
+
+#if WITH_OPENCV
+
+#include "PreOpenCVHeaders.h"
+#include "opencv2/calib3d.hpp"
+#include "PostOpenCVHeaders.h"
+
+#endif	// WITH_OPENCV
+
+static TAutoConsoleVariable<int> CVarUseLegacySphericalSolver(TEXT("CameraCalibration.UseLegacySphericalSolver"), 0, TEXT("If set, the legacy OpenCV spherical solver will be used"));
+
+DEFINE_LOG_CATEGORY_STATIC(LogCameraCalibrationSolver, Log, All);
+
+#if WITH_OPENCV
+
+double FCameraCalibrationSolver::CalibrateCamera(
+	const TSubclassOf<ULensModel> LensModel,
+	const TArray<TArray<FVector>>& InObjectPoints,
+	const TArray<TArray<FVector2D>>& InImagePoints,
+	const FIntPoint ImageSize,
+	FVector2f& InOutFocalLength,
+	FVector2f& InOutImageCenter,
+	TArray<float>& OutDistCoeffs,
+	double PixelAspect,
+	int SolverFlags)
+{
+	const int NumImages = InObjectPoints.Num();
+
+	// Create an array to store the number of points in each image
+	cv::Mat NumPointsMat = cv::Mat(1, NumImages, CV_32S);
+
+	int NumTotalPoints = 0;
+	int MaxPoints = 0;
+	for (int ImageIndex = 0; ImageIndex < NumImages; ++ImageIndex)
+	{
+		const int NumPointsInImage = InObjectPoints[ImageIndex].Num();
+
+		NumPointsMat.at<int>(ImageIndex) = NumPointsInImage;
+		NumTotalPoints += NumPointsInImage;
+
+		MaxPoints = MAX(MaxPoints, NumPointsInImage);
+	}
+
+	cv::Mat ObjectPointsMat = cv::Mat(1, NumTotalPoints, CV_64FC3);
+	cv::Mat ImagePointsMat = cv::Mat(1, NumTotalPoints, CV_64FC2);
+
+	// Reorganize the 3D and 2D points from the input arrays or arrays to be laid out linearly in memory in two cv::Mat objects
+	GatherPoints(InObjectPoints, InImagePoints, ObjectPointsMat, ImagePointsMat);
+
+	double RMSE = 0.0;
+
+	cv::Mat CameraMatrix = cv::Mat::eye(3, 3, CV_64F);
+
+	CameraMatrix.at<double>(0, 0) = InOutFocalLength.X;
+	CameraMatrix.at<double>(1, 1) = InOutFocalLength.Y;
+	CameraMatrix.at<double>(0, 2) = InOutImageCenter.X;
+	CameraMatrix.at<double>(1, 2) = InOutImageCenter.Y;
+
+	OutDistCoeffs.Empty();
+
+	const int NumDistortionCoefficients = LensModel->GetDefaultObject<ULensModel>()->GetNumParameters();
+	cv::Mat DistCoeffs = cv::Mat(1, NumDistortionCoefficients, CV_64F);
+
+	TArray<float> DefaultCoefficients;
+	LensModel->GetDefaultObject<ULensModel>()->GetDefaultParameterArray(DefaultCoefficients);
+
+	for (int CoeffIndex = 0; CoeffIndex < NumDistortionCoefficients; ++CoeffIndex)
+	{
+		DistCoeffs.at<double>(CoeffIndex) = DefaultCoefficients[CoeffIndex];
+	}
+
+	if (LensModel == UAnamorphicLensModel::StaticClass())
+	{
+		DistCoeffs.at<double>(0) = PixelAspect;
+	}
+
+	cv::Size CvImageSize = cv::Size(ImageSize.X, ImageSize.Y);
+
+	// If the cvar is set, and the lens model is spherical, use the legacy version of the solver which simply calls cv::calibrateCamera()
+	// Note: this can be removed in the future after some testing shows that the hand-written version of the solver works identically (if not better) compared to the legacy version
+	if (CVarUseLegacySphericalSolver.GetValueOnAnyThread() && LensModel == USphericalLensModel::StaticClass())
+	{
+		std::vector<cv::Mat> Rvecs;
+		std::vector<cv::Mat> Tvecs;
+
+		std::vector<std::vector<cv::Point2f>> Samples2d;
+		std::vector<std::vector<cv::Point3f>> Samples3d;
+
+		Samples2d.reserve(InImagePoints.Num());
+		Samples3d.reserve(InObjectPoints.Num());
+
+		for (const TArray<FVector>& Image : InObjectPoints)
+		{
+			std::vector<cv::Point3f> Points3d;
+			Points3d.reserve(Image.Num());
+
+			for (const FVector& Point3d : Image)
+			{
+				Points3d.push_back(cv::Point3f(Point3d.X, Point3d.Y, Point3d.Z));
+			}
+
+			Samples3d.push_back(Points3d);
+		}
+
+		for (const TArray<FVector2D>& Image : InImagePoints)
+		{
+			std::vector<cv::Point2f> Points2d;
+			Points2d.reserve(Image.Num());
+
+			for (const FVector2D& Point2d : Image)
+			{
+				Points2d.push_back(cv::Point2f(Point2d.X, Point2d.Y));
+			}
+
+			Samples2d.push_back(Points2d);
+		}
+
+		RMSE = cv::calibrateCamera(
+			Samples3d,
+			Samples2d,
+			CvImageSize,
+			CameraMatrix,
+			DistCoeffs,
+			Rvecs,
+			Tvecs,
+			cv::noArray(),
+			cv::noArray(),
+			cv::noArray(),
+			SolverFlags
+		);
+
+		// Set the output intrinsics and distortion parameters to the final values calculated by the solver
+		InOutFocalLength.X = CameraMatrix.at<double>(0, 0);
+		InOutFocalLength.Y = CameraMatrix.at<double>(1, 1);
+		InOutImageCenter.X = CameraMatrix.at<double>(0, 2);
+		InOutImageCenter.Y = CameraMatrix.at<double>(1, 2);
+
+		// The spherical distortion coefficients in our model are in a different order than they appear in the solver, so the results need to be rearranged
+		OutDistCoeffs.Add(DistCoeffs.at<double>(0));
+		OutDistCoeffs.Add(DistCoeffs.at<double>(1));
+		OutDistCoeffs.Add(DistCoeffs.at<double>(4));
+		OutDistCoeffs.Add(DistCoeffs.at<double>(2));
+		OutDistCoeffs.Add(DistCoeffs.at<double>(3));
+
+		return RMSE;
+	}
+
+	// If the flag to use a starting guess for the camera intrinsics is not set, calculate some initial values for the intrinsic parameters
+	if (!(SolverFlags & cv::CALIB_USE_INTRINSIC_GUESS))
+	{
+		InitCameraIntrinsics(ObjectPointsMat, ImagePointsMat, NumPointsMat, CvImageSize, CameraMatrix);
+	}
+
+	// Initialize the solver with the number of parameters to solve, and the maximum number of iterations to run
+	const int NumIntrinsics = NumDistortionCoefficients + 4; // Includes Fx, Fy, Cx, and Cy
+	const int NumExtrinsics = 6; // 3 for rotation vector, 3 for translation vector
+	const int NumParamsToSolve = NumIntrinsics + (NumImages * NumExtrinsics);
+
+	constexpr int NumErrors = 0;
+	constexpr int MaxIterations = 30;
+
+	FLevMarqSolver Solver(NumParamsToSolve, NumErrors, MaxIterations);
+
+	// Initialize the solver with the starting values for the intrinsic parameters and distortion parameters
+	double* SolverParams = Solver.Params.ptr<double>();
+	SolverParams[0] = CameraMatrix.at<double>(0, 0);
+	SolverParams[1] = CameraMatrix.at<double>(1, 1);
+	SolverParams[2] = CameraMatrix.at<double>(0, 2);
+	SolverParams[3] = CameraMatrix.at<double>(1, 2);
+
+	FMemory::Memcpy(SolverParams + 4, DistCoeffs.ptr<double>(), NumDistortionCoefficients * sizeof(double));
+
+	// Instruct the solver to ignore some parameters when running its solve
+	uchar* Mask = Solver.Mask.ptr<uchar>();
+	if (SolverFlags & cv::CALIB_FIX_FOCAL_LENGTH)
+	{
+		Mask[0] = 0;
+		Mask[1] = 0;
+	}
+	if (SolverFlags & cv::CALIB_FIX_PRINCIPAL_POINT)
+	{
+		Mask[2] = 0;
+		Mask[3] = 0;
+	}
+
+	if (LensModel == UAnamorphicLensModel::StaticClass())
+	{
+		Mask[4] = 0; // We do not want to solve for pixel aspect
+	}
+
+	// Initialize the starting guess for the camera's extrinsic parameters in each image
+	int ObjectPointsIndex = 0;
+	for (int ImageIndex = 0; ImageIndex < NumImages; ImageIndex++)
+	{
+		const int NumImagePoints = NumPointsMat.at<int>(ImageIndex);
+
+		// Get the object and image points for this image
+		cv::Mat ObjectPointsInImage = ObjectPointsMat.colRange(ObjectPointsIndex, ObjectPointsIndex + NumImagePoints);
+		cv::Mat ImagePointsInImage = ImagePointsMat.colRange(ObjectPointsIndex, ObjectPointsIndex + NumImagePoints);
+
+		ObjectPointsIndex += NumImagePoints;
+
+		// Get a view to parameters used by the solver for the rotation and translation vectors for this image
+		const int ExtrinsicOffset = NumIntrinsics + (ImageIndex * NumExtrinsics);
+		cv::Mat Rotation = Solver.Params.rowRange(ExtrinsicOffset, ExtrinsicOffset + 3);
+		cv::Mat Translation = Solver.Params.rowRange(ExtrinsicOffset + 3, ExtrinsicOffset + 6);
+
+		InitCameraExtrinsics(LensModel, ObjectPointsInImage, ImagePointsInImage, CameraMatrix, DistCoeffs, CvImageSize, Rotation, Translation, SolverFlags);
+	}
+
+	double ReprojectionError = 0.0;
+
+	cv::Mat Jacobian(MaxPoints * 2, NumExtrinsics + NumIntrinsics, CV_64FC1, cv::Scalar(0));
+	cv::Mat Diffs(MaxPoints * 2, 1, CV_64FC1);
+
+	while (true)
+	{
+		// Update the solver
+		bool bShouldProceed = Solver.UpdateAlt();
+		bool bComputeJacobian = Solver.State == FLevMarqSolver::ESolverState::ComputeJacobian;
+
+		// Update the camera matrix and distortion parameters with the latest values of the parameters from the solver
+		CameraMatrix.at<double>(0, 0) = SolverParams[0];
+		CameraMatrix.at<double>(1, 1) = SolverParams[1];
+		CameraMatrix.at<double>(0, 2) = SolverParams[2];
+		CameraMatrix.at<double>(1, 2) = SolverParams[3];
+
+		FMemory::Memcpy(DistCoeffs.ptr<double>(), SolverParams + 4, NumDistortionCoefficients * sizeof(double));
+
+		// If the solver determined that it no longer needs to proceed, break out of the loop
+		if (!bShouldProceed)
+		{
+			break;
+		}
+
+		ReprojectionError = 0;
+
+		ObjectPointsIndex = 0;
+		for (int ImageIndex = 0; ImageIndex < NumImages; ImageIndex++)
+		{
+			int NumImagePoints = NumPointsMat.at<int>(ImageIndex);
+
+			// Get the object and image points for this image
+			cv::Mat ObjectPointsInImage = cv::Mat(ObjectPointsMat.colRange(ObjectPointsIndex, ObjectPointsIndex + NumImagePoints));
+			cv::Mat ImagePointsInImage = cv::Mat(ImagePointsMat.colRange(ObjectPointsIndex, ObjectPointsIndex + NumImagePoints));
+
+			ObjectPointsIndex += NumImagePoints;
+
+			// Get the rotation and translation vectors for this image
+			const int ExtrinsicOffset = NumIntrinsics + (ImageIndex * NumExtrinsics);
+			cv::Mat Rotation = Solver.Params.rowRange(ExtrinsicOffset, ExtrinsicOffset + 3);
+			cv::Mat Translation = Solver.Params.rowRange(ExtrinsicOffset + 3, ExtrinsicOffset + 6);
+
+			Jacobian.resize(NumImagePoints * 2);
+			Diffs.resize(NumImagePoints * 2);
+
+			cv::Mat ProjectedPoints = cv::Mat(Diffs.reshape(2, 1));
+
+			// Project the 3D object points to 2D using the camera extrinsics, intrinsics, and distortion parameters, optionally also solving for the Jacobian matrix
+			if (bComputeJacobian)
+			{
+				ProjectPoints(LensModel, ObjectPointsInImage, Rotation, Translation, CameraMatrix, DistCoeffs, CvImageSize, ProjectedPoints, Jacobian, SolverFlags);
+			}
+			else
+			{
+				ProjectPoints(LensModel, ObjectPointsInImage, Rotation, Translation, CameraMatrix, DistCoeffs, CvImageSize, ProjectedPoints);
+			}
+
+			// Compute the difference between the input image points and the projected points
+			ProjectedPoints = ProjectedPoints - ImagePointsInImage;
+
+			// Compute the jacobian matrices used by the solver for its next update
+			if (bComputeJacobian)
+			{
+				cv::Mat JacobianExtrinsics = Jacobian.colRange(0, NumExtrinsics);
+				cv::Mat JacobianIntrinsics = Jacobian.colRange(NumExtrinsics, NumExtrinsics + NumIntrinsics);
+
+				Solver.JacTJac(cv::Rect(0, 0, NumIntrinsics, NumIntrinsics)) += JacobianIntrinsics.t() * JacobianIntrinsics;
+				Solver.JacTJac(cv::Rect(ExtrinsicOffset, ExtrinsicOffset, 6, 6)) += JacobianExtrinsics.t() * JacobianExtrinsics;
+				Solver.JacTJac(cv::Rect(ExtrinsicOffset, 0, 6, NumIntrinsics)) += JacobianIntrinsics.t() * JacobianExtrinsics;
+
+				Solver.JacTErr.rowRange(0, NumIntrinsics) += JacobianIntrinsics.t() * Diffs;
+				Solver.JacTErr.rowRange(ExtrinsicOffset, ExtrinsicOffset + 6) += JacobianExtrinsics.t() * Diffs;
+			}
+
+			const double ViewError = cv::norm(Diffs, cv::NORM_L2SQR);
+
+			ReprojectionError += ViewError;
+		}
+
+		// Update the solver's error with the latest reprojection error
+		Solver.ErrorNorm = ReprojectionError;
+	}
+
+	RMSE = FMath::Sqrt(ReprojectionError / NumTotalPoints);
+
+	// Set the output intrinsics and distortion parameters to the final values calculated by the solver
+	InOutFocalLength.X = CameraMatrix.at<double>(0, 0);
+	InOutFocalLength.Y = CameraMatrix.at<double>(1, 1);
+	InOutImageCenter.X = CameraMatrix.at<double>(0, 2);
+	InOutImageCenter.Y = CameraMatrix.at<double>(1, 2);
+
+	if (LensModel == UAnamorphicLensModel::StaticClass())
+	{
+		for (int CoeffIndex = 0; CoeffIndex < NumDistortionCoefficients; ++CoeffIndex)
+		{
+			OutDistCoeffs.Add(DistCoeffs.at<double>(CoeffIndex));
+		}
+	}
+	else
+	{
+		// The spherical distortion coefficients in our model are in a different order than they appear in the solver, so the results need to be rearranged
+		OutDistCoeffs.Add(DistCoeffs.at<double>(0));
+		OutDistCoeffs.Add(DistCoeffs.at<double>(1));
+		OutDistCoeffs.Add(DistCoeffs.at<double>(4));
+		OutDistCoeffs.Add(DistCoeffs.at<double>(2));
+		OutDistCoeffs.Add(DistCoeffs.at<double>(3));
+	}
+
+	return RMSE;
+}
+
+void FCameraCalibrationSolver::InitCameraIntrinsics(
+	const cv::Mat& ObjectPoints,
+	const cv::Mat& ImagePoints,
+	const cv::Mat& NumPoints,
+	const cv::Size ImageSize,
+	cv::Mat& CameraMatrix)
+{
+	const int NumImages = NumPoints.cols;
+
+	const double Cx = (ImageSize.width - 1) * 0.5;
+	const double Cy = (ImageSize.height - 1) * 0.5;
+
+	cv::Mat MatrixA = cv::Mat(2 * NumImages, 2, CV_64F);
+	cv::Mat VecB = cv::Mat(2 * NumImages, 1, CV_64F);
+
+	// extract vanishing points in order to obtain initial value for the focal length
+	int ObjectPointsIndex = 0;
+	for (int ImageIndex = 0; ImageIndex < NumImages; ++ImageIndex)
+	{
+		const int NumImagePoints = NumPoints.ptr<int>()[ImageIndex];
+
+		const cv::Mat ObjectPointsInImage = ObjectPoints.colRange(ObjectPointsIndex, ObjectPointsIndex + NumImagePoints);
+		const cv::Mat ImagePointsInImage = ImagePoints.colRange(ObjectPointsIndex, ObjectPointsIndex + NumImagePoints);
+
+		ObjectPointsIndex += NumImagePoints;
+
+		cv::Mat HomographyMat = cv::findHomography(ObjectPointsInImage, ImagePointsInImage);
+
+		double* Homography = HomographyMat.ptr<double>();
+
+		Homography[0] -= Homography[6] * Cx;
+		Homography[1] -= Homography[7] * Cx;
+		Homography[2] -= Homography[8] * Cx;
+		Homography[3] -= Homography[6] * Cy;
+		Homography[4] -= Homography[7] * Cy;
+		Homography[5] -= Homography[8] * Cy;
+
+		double H[3];
+		double V[3];
+		double D1[3];
+		double D2[3];
+		double N[4] = { 0 };
+
+		for (int Index = 0; Index < 3; ++Index)
+		{
+			double T0 = Homography[Index * 3];
+			double T1 = Homography[Index * 3 + 1];
+
+			H[Index] = T0;
+			V[Index] = T1;
+
+			D1[Index] = (T0 + T1) * 0.5;
+			D2[Index] = (T0 - T1) * 0.5;
+
+			N[0] += H[Index] * H[Index];
+			N[1] += V[Index] * V[Index];
+			N[2] += D1[Index] * D1[Index];
+			N[3] += D2[Index] * D2[Index];
+		}
+
+		N[0] = FMath::InvSqrt(N[0]);
+		N[1] = FMath::InvSqrt(N[1]);
+		N[2] = FMath::InvSqrt(N[2]);
+		N[3] = FMath::InvSqrt(N[3]);
+
+		for (int Index = 0; Index < 3; ++Index)
+		{
+			H[Index] *= N[0];
+			V[Index] *= N[1];
+
+			D1[Index] *= N[2];
+			D2[Index] *= N[3];
+		}
+
+		double* MatrixAPtr = MatrixA.ptr<double>(ImageIndex * 2);
+		double* VecBPtr = VecB.ptr<double>(ImageIndex * 2);
+
+		MatrixAPtr[0] = H[0] * V[0];
+		MatrixAPtr[1] = H[1] * V[1];
+		MatrixAPtr[2] = D1[0] * D2[0];
+		MatrixAPtr[3] = D1[1] * D2[1];
+
+		VecBPtr[0] = -H[2] * V[2];
+		VecBPtr[1] = -D1[2] * D2[2];
+	}
+
+	double FxFy[2] = { 0 };
+	cv::Mat FxFyMat = cv::Mat(2, 1, CV_64F, FxFy);
+
+	cv::solve(MatrixA, VecB, FxFyMat, cv::DECOMP_NORMAL + cv::DECOMP_SVD);
+
+	const double Fx = FMath::Sqrt(FMath::Abs(1.0 / FxFy[0]));
+	const double Fy = FMath::Sqrt(FMath::Abs(1.0 / FxFy[1]));
+
+	CameraMatrix.at<double>(0, 0) = Fx;
+	CameraMatrix.at<double>(1, 1) = Fy;
+	CameraMatrix.at<double>(0, 2) = Cx;
+	CameraMatrix.at<double>(1, 2) = Cy;
+}
+
+void FCameraCalibrationSolver::InitCameraExtrinsics(
+	const TSubclassOf<ULensModel> LensModel,
+	const cv::Mat& ObjectPoints,
+	const cv::Mat& ImagePoints,
+	const cv::Mat& CameraMatrix,
+	const cv::Mat& DistCoeffs,
+	const cv::Size ImageSize,
+	cv::Mat& Rotation,
+	cv::Mat& Translation,
+	const int SolverFlags)
+{
+	const int NumPoints = ObjectPoints.total();
+
+	// Normalize image points (unapply the intrinsic matrix transformation)
+	cv::Mat ZeroDistortion = cv::Mat::zeros(5, 1, CV_64F);
+	cv::Mat UndistortedPoints = cv::Mat(1, NumPoints, CV_64FC2);
+	cv::undistortPoints(ImagePoints, UndistortedPoints, CameraMatrix, ZeroDistortion);
+
+	// Average the X, Y, and Z channels for every object point in the set
+	cv::Mat ObjectPointAvg = cv::Mat(1, 3, CV_64F, cv::mean(ObjectPoints).val);
+
+	// Reshape the ObjectPoints mat into a 1-channel (NumPoints x 3) matrix, where each row is [ X Y Z ]
+	cv::Mat ObjectPoints3Column = ObjectPoints.reshape(1, NumPoints);
+
+	// Multiply the ObjectPoints Transpose by itself, subtracting the average values that were computed
+	constexpr bool bTransposeFirstMatrix = true;
+	cv::Mat MM = cv::Mat(3, 3, CV_64F);
+
+	cv::mulTransposed(ObjectPoints3Column, MM, bTransposeFirstMatrix, ObjectPointAvg);
+
+	// Perform singular value decomposition on the resulting matrix
+	double V[9] = { 0 };
+	cv::Mat MatrixV = cv::Mat(3, 3, CV_64F, V);
+
+	double W[3] = { 0 };
+	cv::Mat MatrixW = cv::Mat(3, 1, CV_64F, W);
+
+	cv::SVDecomp(MM, MatrixW, cv::noArray(), MatrixV, cv::SVD::MODIFY_A + cv::SVD::FULL_UV);
+
+	// a planar structure case (all object points lie in the same plane)
+	if (W[2] / W[1] < 1e-3)
+	{
+		cv::Mat RotationMatrix = MatrixV;
+
+		if ((V[2] * V[2] + V[5] * V[5]) < 1e-10)
+		{
+			RotationMatrix = cv::Mat::eye(3, 3, CV_64F);
+		}
+
+		if (cv::determinant(RotationMatrix) < 0)
+		{
+			RotationMatrix = RotationMatrix.mul(-1);
+		}
+
+		Translation = -RotationMatrix * ObjectPointAvg.t();
+
+		cv::Mat TransformedPoints = cv::Mat(1, NumPoints, CV_64FC2);
+
+		for (int PointIndex = 0; PointIndex < NumPoints; PointIndex++)
+		{
+			const cv::Point3d& ObjectPoint = ObjectPoints.at<cv::Point3d>(PointIndex);
+			cv::Point2d& TransformedPoint = TransformedPoints.at<cv::Point2d>(PointIndex);
+
+			const double* RPtr = RotationMatrix.ptr<double>();
+			const double* TPtr = Translation.ptr<double>();
+
+			TransformedPoint.x = (ObjectPoint.x * RPtr[0]) + (ObjectPoint.y * RPtr[1]) + (ObjectPoint.z * RPtr[2]) + TPtr[0];
+			TransformedPoint.y = (ObjectPoint.x * RPtr[3]) + (ObjectPoint.y * RPtr[4]) + (ObjectPoint.z * RPtr[5]) + TPtr[1];
+		}
+
+		cv::Mat Homography = cv::findHomography(TransformedPoints, UndistortedPoints);
+
+		if (cv::checkRange(Homography, true))
+		{
+			cv::Mat H1 = Homography.col(0);
+			cv::Mat H2 = Homography.col(1);
+			cv::Mat H3 = Homography.col(2);
+
+			const double H1Norm = FMath::Max(cv::norm(H1), DBL_EPSILON);
+			const double H2Norm = FMath::Max(cv::norm(H2), DBL_EPSILON);
+
+			H1 = H1.mul(1.0 / H1Norm);
+			H2 = H2.mul(1.0 / H2Norm);
+			cv::Mat H3Norm = H3.mul(2.0 / (H1Norm + H2Norm));
+
+			H3 = H1.cross(H2);
+
+			RotationMatrix = Homography * RotationMatrix;
+
+			Translation = Homography * Translation;
+			Translation += H3Norm;
+		}
+		else
+		{
+			RotationMatrix = cv::Mat::eye(3, 3, CV_64F);
+			Translation.setTo(cv::Scalar::all(0));
+		}
+
+		cv::Rodrigues(RotationMatrix, Rotation);
+	}
+	else // non-planar structure, not currently supported
+	{
+		cv::Mat RotationMatrix = cv::Mat::eye(3, 3, CV_64F);
+		cv::Rodrigues(RotationMatrix, Rotation);
+
+		Translation.setTo(cv::Scalar::all(0));
+	}
+
+	// Refine extrinsic parameters
+	constexpr int NumExtrinsicParams = 6;
+	const int NumErrors = NumPoints * 2;
+	constexpr int MaxIterations = 30;
+
+	FLevMarqSolver Solver(NumExtrinsicParams, NumErrors, MaxIterations, FLevMarqSolver::ESymmetryMode::UpperToLower);
+
+	double* SolverParams = Solver.Params.ptr<double>();
+
+	SolverParams[0] = Rotation.at<double>(0);
+	SolverParams[1] = Rotation.at<double>(1);
+	SolverParams[2] = Rotation.at<double>(2);
+
+	SolverParams[3] = Translation.at<double>(0);
+	SolverParams[4] = Translation.at<double>(1);
+	SolverParams[5] = Translation.at<double>(2);
+
+	while (true)
+	{
+		bool bShouldProceed = Solver.Update();
+		bool bComputeJacobian = Solver.State == FLevMarqSolver::ESolverState::ComputeJacobian;
+
+		Rotation.at<double>(0) = SolverParams[0];
+		Rotation.at<double>(1) = SolverParams[1];
+		Rotation.at<double>(2) = SolverParams[2];
+
+		Translation.at<double>(0) = SolverParams[3];
+		Translation.at<double>(1) = SolverParams[4];
+		Translation.at<double>(2) = SolverParams[5];
+
+		if (!bShouldProceed)
+		{
+			break;
+		}
+
+		cv::Mat ProjectedPoints = Solver.Error.reshape(2, 1);
+
+		if (bComputeJacobian)
+		{
+			ProjectPoints(LensModel, ObjectPoints, Rotation, Translation, CameraMatrix, DistCoeffs, ImageSize, ProjectedPoints, Solver.Jacobian, SolverFlags);
+		}
+		else
+		{
+			ProjectPoints(LensModel, ObjectPoints, Rotation, Translation, CameraMatrix, DistCoeffs, ImageSize, ProjectedPoints);
+		}
+
+		ProjectedPoints = ProjectedPoints - ImagePoints;
+	}
+}
+
+void FCameraCalibrationSolver::ProjectPoints(
+	const TSubclassOf<ULensModel> LensModel,
+	const cv::Mat& ObjectPoints,
+	const cv::Mat& Rotation,
+	const cv::Mat& Translation,
+	const cv::Mat& CameraMatrix,
+	const cv::Mat& DistCoeffs,
+	const cv::Size ImageSize,
+	cv::Mat& ProjectedPoints)
+{
+	cv::Mat Jacobian;
+	const int SolverFlags = 0;
+	ProjectPoints(LensModel, ObjectPoints, Rotation, Translation, CameraMatrix, DistCoeffs, ImageSize, ProjectedPoints, Jacobian, SolverFlags);
+}
+
+void FCameraCalibrationSolver::ProjectPoints(
+	const TSubclassOf<ULensModel> LensModel,
+	const cv::Mat& ObjectPoints,
+	const cv::Mat& Rotation,
+	const cv::Mat& Translation,
+	const cv::Mat& CameraMatrix,
+	const cv::Mat& DistCoeffs,
+	const cv::Size ImageSize,
+	cv::Mat& ProjectedPoints,
+	cv::Mat& Jacobian,
+	int SolverFlags)
+{
+	// Project the 3D points to 2D using the transformation matrix, camera matrix, and distortion equation
+	if (LensModel == USphericalLensModel::StaticClass())
+	{
+		ProjectPointsSpherical(
+			ObjectPoints,
+			Rotation,
+			Translation,
+			CameraMatrix,
+			DistCoeffs,
+			ProjectedPoints,
+			Jacobian,
+			SolverFlags);
+	}
+	else if (LensModel == UAnamorphicLensModel::StaticClass())
+	{
+		ProjectPointsAnamorphic(
+			ObjectPoints,
+			Rotation,
+			Translation,
+			CameraMatrix,
+			DistCoeffs,
+			ImageSize,
+			ProjectedPoints,
+			Jacobian,
+			SolverFlags);
+	}
+}
+
+void FCameraCalibrationSolver::ProjectPointsAnamorphic(
+	const cv::Mat& ObjectPoints,
+	const cv::Mat& Rotation,
+	const cv::Mat& Translation,
+	const cv::Mat& CameraMatrix,
+	const cv::Mat& DistCoeffs,
+	const cv::Size ImageSize,
+	cv::Mat& ProjectedPoints,
+	cv::Mat& Jacobian,
+	int SolverFlags)
+{
+	const bool bComputeJacobian = !Jacobian.empty();
+
+	// Get the rotation vector/matrix into the correct format and also compute the matrix of partial derivatives
+	double R[9];
+	cv::Mat RotationMatrix = cv::Mat(3, 3, CV_64F, R);
+
+	double JacR[27];
+	cv::Mat RotationJacobian = cv::Mat(3, 9, CV_64F, JacR);
+
+	cv::Rodrigues(Rotation, RotationMatrix, RotationJacobian);
+
+	const double* TranslationVector = Translation.ptr<double>();
+
+	// Extract Camera Intrinsics
+	const double Fx = CameraMatrix.at<double>(0, 0);
+	const double Fy = CameraMatrix.at<double>(1, 1);
+	const double Cx = CameraMatrix.at<double>(0, 2);
+	const double Cy = CameraMatrix.at<double>(1, 2);
+
+	// Extract Distortion Parameters
+	const double* DistortionParams = DistCoeffs.ptr<double>();
+
+	const double PixelAspect = DistortionParams[0];
+	const double CX02 = DistortionParams[1];
+	const double CX04 = DistortionParams[2];
+	const double CX22 = DistortionParams[3];
+	const double CX24 = DistortionParams[4];
+	const double CX44 = DistortionParams[5];
+	const double CY02 = DistortionParams[6];
+	const double CY04 = DistortionParams[7];
+	const double CY22 = DistortionParams[8];
+	const double CY24 = DistortionParams[9];
+	const double CY44 = DistortionParams[10];
+	const double SqueezeX = DistortionParams[11];
+	const double SqueezeY = DistortionParams[12];
+	const double LensRotationDeg = DistortionParams[13];
+
+	// Pre-compute values that will be re-used in the distortion equation for every point
+	const double LensRotationRad = LensRotationDeg * (PI / 180.0);
+	const double CosLensRotation = FMath::Cos(LensRotationRad);
+	const double SinLensRotation = FMath::Sin(LensRotationRad);
+
+	const double FilmbackWidth = ImageSize.width;
+	const double FilmbackHeight = ImageSize.height;
+
+	const double FilmbackRadius = (0.5) * FMath::Sqrt((FilmbackWidth * FilmbackWidth) + (FilmbackHeight * FilmbackHeight));
+
+	const int NumPoints = ObjectPoints.total();
+
+	for (int PointIndex = 0; PointIndex < NumPoints; ++PointIndex)
+	{
+		const cv::Point3d& WorldPoint = ObjectPoints.at<cv::Point3d>(PointIndex);
+		cv::Point2d& ProjectedPoint = ProjectedPoints.at<cv::Point2d>(PointIndex);
+
+		// Transform the 3D point from the world coordinate system to our camera coordinate system
+		cv::Point3d CameraPoint;
+		CameraPoint.x = (WorldPoint.x * R[0]) + (WorldPoint.y * R[1]) + (WorldPoint.z * R[2]) + TranslationVector[0];
+		CameraPoint.y = (WorldPoint.x * R[3]) + (WorldPoint.y * R[4]) + (WorldPoint.z * R[5]) + TranslationVector[1];
+		CameraPoint.z = (WorldPoint.x * R[6]) + (WorldPoint.y * R[7]) + (WorldPoint.z * R[8]) + TranslationVector[2];
+
+		// Divide x and y by the z-coordinate to normalize the point before projection
+		CameraPoint.z = CameraPoint.z ? (1.0 / CameraPoint.z) : 1.0;
+
+		CameraPoint.x *= CameraPoint.z;
+		CameraPoint.y *= CameraPoint.z;
+
+		// Project the point onto the image plane using our camera intrinsics
+		cv::Point2d ImagePoint;
+		ImagePoint.x = CameraPoint.x * Fx + Cx;
+		ImagePoint.y = CameraPoint.y * Fy + Cy;
+
+		// Compute the diagonally normalized point using the image center and the filmback radius
+		cv::Point2d DiagNormPoint;
+		DiagNormPoint.x = (ImagePoint.x - Cx) / FilmbackRadius;
+		DiagNormPoint.y = (ImagePoint.y - Cy) / FilmbackRadius;
+
+		// Rotate the point using the lens rotation model parameter
+		cv::Point2d RotatedPoint;
+		RotatedPoint.x = (DiagNormPoint.x * CosLensRotation) + (DiagNormPoint.y * SinLensRotation);
+		RotatedPoint.y = (DiagNormPoint.y * CosLensRotation) - (DiagNormPoint.x * SinLensRotation);
+
+		// Divide the x-coordinate by the pixel aspect
+		RotatedPoint.x = RotatedPoint.x / PixelAspect;
+
+		// The anamorphic model is intended to undistort points. Here we use it to distort points, so we run an iterative solution to reverse the model
+		cv::Point2d DistortedPoint;
+		DistortedPoint.x = RotatedPoint.x;
+		DistortedPoint.y = RotatedPoint.y;
+
+		double RadiusSq = 0.0;
+		double Radius4th = 0.0;
+		double Angle = 0.0;
+		double DistortionX = 0.0;
+		double DistortionY = 0.0;
+		double CosAngle2 = 0.0;
+		double CosAngle4 = 0.0;
+		for (int Index = 0; Index < 10; ++Index)
+		{
+			// Compute the polar coordinates (radius and angle) for the current point
+			RadiusSq = (DistortedPoint.x * DistortedPoint.x) + (DistortedPoint.y * DistortedPoint.y);
+			Radius4th = RadiusSq * RadiusSq;
+			Angle = FMath::Atan2(DistortedPoint.y, DistortedPoint.x);
+			CosAngle2 = FMath::Cos(Angle * 2);
+			CosAngle4 = FMath::Cos(Angle * 4);
+
+			// Compute the distortion contribution in the x and y directions
+			DistortionX = 1 + (CX02 * RadiusSq) + (CX04 * Radius4th) + (CX22 * RadiusSq * CosAngle2) + (CX24 * Radius4th * CosAngle2) + (CX44 * Radius4th * CosAngle4);
+			DistortionY = 1 + (CY02 * RadiusSq) + (CY04 * Radius4th) + (CY22 * RadiusSq * CosAngle2) + (CY24 * Radius4th * CosAngle2) + (CY44 * Radius4th * CosAngle4);
+
+			// Distort the undistorted point
+			DistortedPoint.x = (RotatedPoint.x / DistortionX) / SqueezeX;
+			DistortedPoint.y = (RotatedPoint.y / DistortionY) / SqueezeY;
+		}
+
+		// Undo the model's lens rotation that we applied earlier
+		cv::Point2d DistortedPointRot;
+		DistortedPointRot.x = (DistortedPoint.x * CosLensRotation) - (DistortedPoint.y * SinLensRotation);
+		DistortedPointRot.y = (DistortedPoint.y * CosLensRotation) + (DistortedPoint.x * SinLensRotation);
+
+		// Re-multiply the x-coordinate by the pixel aspect
+		DistortedPointRot.x = DistortedPointRot.x * PixelAspect;
+
+		// Go back from diagonally normalized coordinates to pixels
+		ProjectedPoint.x = (DistortedPointRot.x * FilmbackRadius) + Cx;
+		ProjectedPoint.y = (DistortedPointRot.y * FilmbackRadius) + Cy;
+
+		if (bComputeJacobian)
+		{
+			const int JacobianStep = Jacobian.cols;
+
+			cv::Mat JacRotation;
+			cv::Mat JacTranslation;
+			cv::Mat JacFocalLength;
+			cv::Mat JacImageCenter;
+			cv::Mat JacDistortion;
+			SubdivideJacobian(Jacobian, JacRotation, JacTranslation, JacFocalLength, JacImageCenter, JacDistortion, SolverFlags);
+
+			if (!JacDistortion.empty())
+			{
+				const double DDistortedPointXDCX02 = (RotatedPoint.x * (-1.0 / (DistortionX * DistortionX)) * (RadiusSq)) / SqueezeX;
+				const double DDistortedPointXDCX04 = (RotatedPoint.x * (-1.0 / (DistortionX * DistortionX)) * (Radius4th)) / SqueezeX;
+				const double DDistortedPointXDCX22 = (RotatedPoint.x * (-1.0 / (DistortionX * DistortionX)) * (RadiusSq * CosAngle2)) / SqueezeX;
+				const double DDistortedPointXDCX24 = (RotatedPoint.x * (-1.0 / (DistortionX * DistortionX)) * (Radius4th * CosAngle2)) / SqueezeX;
+				const double DDistortedPointXDCX44 = (RotatedPoint.x * (-1.0 / (DistortionX * DistortionX)) * (Radius4th * CosAngle4)) / SqueezeX;
+
+				const double DDistortedPointYDCY02 = (RotatedPoint.y * (-1.0 / (DistortionY * DistortionY)) * (RadiusSq)) / SqueezeY;
+				const double DDistortedPointYDCY04 = (RotatedPoint.y * (-1.0 / (DistortionY * DistortionY)) * (Radius4th)) / SqueezeY;
+				const double DDistortedPointYDCY22 = (RotatedPoint.y * (-1.0 / (DistortionY * DistortionY)) * (RadiusSq * CosAngle2)) / SqueezeY;
+				const double DDistortedPointYDCY24 = (RotatedPoint.y * (-1.0 / (DistortionY * DistortionY)) * (Radius4th * CosAngle2)) / SqueezeY;
+				const double DDistortedPointYDCY44 = (RotatedPoint.y * (-1.0 / (DistortionY * DistortionY)) * (Radius4th * CosAngle4)) / SqueezeY;
+
+				double DDistortedPointXDSX = -(RotatedPoint.x / DistortionX) / (SqueezeX * SqueezeX);
+				double DDistortedPointYDSY = -(RotatedPoint.y / DistortionY) / (SqueezeY * SqueezeY);
+
+				double DRotatedPointXDRot = ((DiagNormPoint.x * -SinLensRotation) + (DiagNormPoint.y * CosLensRotation)) * (PI / 180.0) / PixelAspect;
+				double DRotatedPointYDRot = ((DiagNormPoint.y * -SinLensRotation) - (DiagNormPoint.x * CosLensRotation)) * (PI / 180.0);
+
+				double DDistortedPointRotXDRot = ((DRotatedPointXDRot * CosLensRotation - SinLensRotation * (PI / 180.0) * DistortedPoint.x) - (DRotatedPointYDRot * SinLensRotation + CosLensRotation * (PI / 180.0) * DistortedPoint.y)) * PixelAspect;
+				double DDistortedPointRotYDRot = (DRotatedPointYDRot * CosLensRotation - SinLensRotation * (PI / 180.0) * DistortedPoint.y) + (DRotatedPointXDRot * SinLensRotation + CosLensRotation * (PI / 180.0) * DistortedPoint.x);
+
+				double* JacDistortionPtr = JacDistortion.ptr<double>(PointIndex * 2);
+
+				JacDistortionPtr[0] = 0;
+				JacDistortionPtr[1] = FilmbackRadius * CosLensRotation * PixelAspect * DDistortedPointXDCX02;
+				JacDistortionPtr[2] = FilmbackRadius * CosLensRotation * PixelAspect * DDistortedPointXDCX04;
+				JacDistortionPtr[3] = FilmbackRadius * CosLensRotation * PixelAspect * DDistortedPointXDCX22;
+				JacDistortionPtr[4] = FilmbackRadius * CosLensRotation * PixelAspect * DDistortedPointXDCX24;
+				JacDistortionPtr[5] = FilmbackRadius * CosLensRotation * PixelAspect * DDistortedPointXDCX44;
+				JacDistortionPtr[6] = 0;
+				JacDistortionPtr[7] = 0;
+				JacDistortionPtr[8] = 0;
+				JacDistortionPtr[9] = 0;
+				JacDistortionPtr[10] = 0;
+				JacDistortionPtr[11] = FilmbackRadius * CosLensRotation * PixelAspect * DDistortedPointXDSX;
+				JacDistortionPtr[12] = 0;
+				JacDistortionPtr[13] = FilmbackRadius * DDistortedPointRotXDRot;
+
+				JacDistortionPtr[JacobianStep + 0] = 0;
+				JacDistortionPtr[JacobianStep + 1] = 0;
+				JacDistortionPtr[JacobianStep + 2] = 0;
+				JacDistortionPtr[JacobianStep + 3] = 0;
+				JacDistortionPtr[JacobianStep + 4] = 0;
+				JacDistortionPtr[JacobianStep + 5] = 0;
+				JacDistortionPtr[JacobianStep + 6] = FilmbackRadius * CosLensRotation * DDistortedPointYDCY02;
+				JacDistortionPtr[JacobianStep + 7] = FilmbackRadius * CosLensRotation * DDistortedPointYDCY04;
+				JacDistortionPtr[JacobianStep + 8] = FilmbackRadius * CosLensRotation * DDistortedPointYDCY22;
+				JacDistortionPtr[JacobianStep + 9] = FilmbackRadius * CosLensRotation * DDistortedPointYDCY24;
+				JacDistortionPtr[JacobianStep + 10] = FilmbackRadius * CosLensRotation * DDistortedPointYDCY44;
+				JacDistortionPtr[JacobianStep + 11] = 0;
+				JacDistortionPtr[JacobianStep + 12] = FilmbackRadius * CosLensRotation * DDistortedPointYDSY;
+				JacDistortionPtr[JacobianStep + 13] = FilmbackRadius * DDistortedPointRotYDRot;
+			}
+
+			if (!JacImageCenter.empty())
+			{
+				double* JacImageCenterPtr = JacImageCenter.ptr<double>(PointIndex * 2);
+
+				JacImageCenterPtr[0] = 1.0;
+				JacImageCenterPtr[1] = 0.0;
+
+				JacImageCenterPtr[JacobianStep + 0] = 0.0;
+				JacImageCenterPtr[JacobianStep + 1] = 1.0;
+			}
+
+			if (!JacFocalLength.empty())
+			{
+				// Compute the derivative of the undistorted image point with respect to Fx and Fy
+				const cv::Point2d DImagePointXDFxFy = { CameraPoint.x, 0.0 };
+				const cv::Point2d DImagePointYDFxFy = { 0.0, CameraPoint.y };
+
+				// Compute the derivative of the diagonally normalized point with respect to Fx and Fy
+				const cv::Point2d DDiagNormPointXDFxFy = (1.0 / FilmbackRadius) * DImagePointXDFxFy;
+				const cv::Point2d DDiagNormPointYDFxFy = (1.0 / FilmbackRadius) * DImagePointYDFxFy;
+
+				// Compute the derivative of the rotated point with respect to Fx and Fy
+				const cv::Point2d DRotatedPointXDFxFy = (DDiagNormPointXDFxFy * CosLensRotation + DDiagNormPointYDFxFy * SinLensRotation) / PixelAspect;
+				const cv::Point2d DRotatedPointYDFxFy = (DDiagNormPointYDFxFy * CosLensRotation + DDiagNormPointXDFxFy * SinLensRotation);
+
+				// Compute the derivative of the distorted point with respect to Fx and Fy
+				const cv::Point2d DDistortedPointXDFxFy = (DRotatedPointXDFxFy * CosLensRotation - DRotatedPointYDFxFy * SinLensRotation) * PixelAspect;
+				const cv::Point2d DDistortedPointYDFxFy = (DRotatedPointYDFxFy * CosLensRotation + DRotatedPointXDFxFy * SinLensRotation);
+				
+				// Add derivatives to the Jacobian matrix
+				double* JacFocalLengthPtr = JacFocalLength.ptr<double>(PointIndex * 2);
+
+				JacFocalLengthPtr[0] = FilmbackRadius * DDistortedPointXDFxFy.x;
+				JacFocalLengthPtr[1] = FilmbackRadius * DDistortedPointXDFxFy.y;
+
+				JacFocalLengthPtr[JacobianStep + 0] = FilmbackRadius * DDistortedPointYDFxFy.x;
+				JacFocalLengthPtr[JacobianStep + 1] = FilmbackRadius * DDistortedPointYDFxFy.y;
+			}
+
+			if (!JacTranslation.empty())
+			{
+				// Compute the derivative of the undistorted image point with respect to the translation vector [ T0  T1  T2 ]
+				const cv::Point3d DImagePointXDT = { Fx * CameraPoint.z, 0.0,                Fx * (-CameraPoint.x * CameraPoint.z) };
+				const cv::Point3d DImagePointYDT = { 0.0,                Fy * CameraPoint.z, Fy * (-CameraPoint.y * CameraPoint.z) };
+
+				// Compute the derivative of the diagonally normalized point with respect to the translation vector [ T0  T1  T2 ]
+				const cv::Point3d DDiagNormPointXDT = (1.0 / FilmbackRadius) * DImagePointXDT;
+				const cv::Point3d DDiagNormPointYDT = (1.0 / FilmbackRadius) * DImagePointYDT;
+
+				// Compute the derivative of the rotated point with respect to the translation vector [ T0  T1  T2 ]
+				const cv::Point3d DRotatedPointXDT = (DDiagNormPointXDT * CosLensRotation + DDiagNormPointYDT * SinLensRotation) / PixelAspect;
+				const cv::Point3d DRotatedPointYDT = (DDiagNormPointYDT * CosLensRotation + DDiagNormPointXDT * SinLensRotation);
+
+				// Compute the derivative of the distorted point with respect to the translation vector [ T0  T1  T2 ]
+				const cv::Point3d DDistortedPointXDT = (DRotatedPointXDT * CosLensRotation - DRotatedPointYDT * SinLensRotation) * PixelAspect;
+				const cv::Point3d DDistortedPointYDT = (DRotatedPointYDT * CosLensRotation + DRotatedPointXDT * SinLensRotation);
+
+				// Add derivatives to the Jacobian matrix
+				double* JacTranslationPtr = JacTranslation.ptr<double>(PointIndex * 2);
+
+				JacTranslationPtr[0] = FilmbackRadius * DDistortedPointXDT.x;
+				JacTranslationPtr[1] = FilmbackRadius * DDistortedPointXDT.y;
+				JacTranslationPtr[2] = FilmbackRadius * DDistortedPointXDT.z;
+
+				JacTranslationPtr[JacobianStep + 0] = FilmbackRadius * DDistortedPointYDT.x;
+				JacTranslationPtr[JacobianStep + 1] = FilmbackRadius * DDistortedPointYDT.y;
+				JacTranslationPtr[JacobianStep + 2] = FilmbackRadius * DDistortedPointYDT.z;
+			}
+
+			if (!JacRotation.empty())
+			{
+				// Compute the derivatives of the camera points with respect to the rotation vector [ R0  R1  R2]
+				// This requires the use of the rotation jacobian matrix that relates the rotation vector in Rodrigues form to the rotation matrix
+				cv::Point3d DCameraPointXDR;
+				DCameraPointXDR.x = WorldPoint.x * JacR[0] + WorldPoint.y * JacR[1] + WorldPoint.z * JacR[2];
+				DCameraPointXDR.y = WorldPoint.x * JacR[9] + WorldPoint.y * JacR[10] + WorldPoint.z * JacR[11];
+				DCameraPointXDR.z = WorldPoint.x * JacR[18] + WorldPoint.y * JacR[19] + WorldPoint.z * JacR[20];
+
+				cv::Point3d DCameraPointYDR;
+				DCameraPointYDR.x = WorldPoint.x * JacR[3] + WorldPoint.y * JacR[4] + WorldPoint.z * JacR[5];
+				DCameraPointYDR.y = WorldPoint.x* JacR[12] + WorldPoint.y * JacR[13] + WorldPoint.z * JacR[14];
+				DCameraPointYDR.z = WorldPoint.x* JacR[21] + WorldPoint.y * JacR[22] + WorldPoint.z * JacR[23];
+
+				cv::Point3d DCameraPointZDR;
+				DCameraPointZDR.x = WorldPoint.x* JacR[6] + WorldPoint.y * JacR[7] + WorldPoint.z * JacR[8];
+				DCameraPointZDR.y = WorldPoint.x* JacR[15] + WorldPoint.y * JacR[16] + WorldPoint.z * JacR[17];
+				DCameraPointZDR.z = WorldPoint.x* JacR[24] + WorldPoint.y * JacR[25] + WorldPoint.z * JacR[26];
+
+				// Compute the derivative of the undistorted image point with respect to the rotation vector [ R0  R1  R2 ]
+				const cv::Point3d DImagePointXDR = Fx * CameraPoint.z * (DCameraPointXDR - CameraPoint.x * DCameraPointZDR);
+				const cv::Point3d DImagePointYDR = Fy * CameraPoint.z * (DCameraPointYDR - CameraPoint.y * DCameraPointZDR);
+
+				// Compute the derivative of the diagonally normalized point with respect to the rotation vector [ R0  R1  R2 ]
+				const cv::Point3d DDiagNormPointXDR = (1.0 / FilmbackRadius) * DImagePointXDR;
+				const cv::Point3d DDiagNormPointYDR = (1.0 / FilmbackRadius) * DImagePointYDR;
+
+				// Compute the derivative of the rotated point with respect to the rotation vector [ R0  R1  R2 ]
+				const cv::Point3d DRotatedPointXDR = (DDiagNormPointXDR * CosLensRotation + DDiagNormPointYDR * SinLensRotation) / PixelAspect;
+				const cv::Point3d DRotatedPointYDR = (DDiagNormPointYDR * CosLensRotation + DDiagNormPointXDR * SinLensRotation);
+
+				// Compute the derivative of the distorted point with respect to the rotation vector [ R0  R1  R2 ]
+				const cv::Point3d DDistortedPointXDR = (DRotatedPointXDR * CosLensRotation - DRotatedPointYDR * SinLensRotation) * PixelAspect;
+				const cv::Point3d DDistortedPointYDR = (DRotatedPointYDR * CosLensRotation + DRotatedPointXDR * SinLensRotation);
+
+				// Add derivatives to the Jacobian matrix
+				double* JacRotationPtr = JacRotation.ptr<double>(PointIndex * 2);
+
+				JacRotationPtr[0] = FilmbackRadius * DDistortedPointXDR.x;
+				JacRotationPtr[1] = FilmbackRadius * DDistortedPointXDR.y;
+				JacRotationPtr[2] = FilmbackRadius * DDistortedPointXDR.z;
+
+				JacRotationPtr[JacobianStep + 0] = FilmbackRadius * DDistortedPointYDR.x;
+				JacRotationPtr[JacobianStep + 1] = FilmbackRadius * DDistortedPointYDR.y;
+				JacRotationPtr[JacobianStep + 2] = FilmbackRadius * DDistortedPointYDR.z;
+			}
+		}
+	}
+}
+
+void FCameraCalibrationSolver::ProjectPointsSpherical(
+	const cv::Mat& ObjectPoints,
+	const cv::Mat& Rotation,
+	const cv::Mat& Translation,
+	const cv::Mat& CameraMatrix,
+	const cv::Mat& DistCoeffs,
+	cv::Mat& ProjectedPoints,
+	cv::Mat& Jacobian,
+	int SolverFlags)
+{
+	const bool bComputeJacobian = !Jacobian.empty();
+
+	// Get the rotation vector/matrix into the correct format and also compute the matrix of partial derivatives
+	double R[9];
+	cv::Mat RotationMatrix = cv::Mat(3, 3, CV_64F, R);
+
+	double JacR[27];
+	cv::Mat RotationJacobian = cv::Mat(3, 9, CV_64F, JacR);
+
+	cv::Rodrigues(Rotation, RotationMatrix, RotationJacobian);
+
+	const double* TranslationVector = Translation.ptr<double>();
+
+	// Extract Camera Intrinsics
+	const double Fx = CameraMatrix.at<double>(0, 0);
+	const double Fy = CameraMatrix.at<double>(1, 1);
+	const double Cx = CameraMatrix.at<double>(0, 2);
+	const double Cy = CameraMatrix.at<double>(1, 2);
+
+	// Extract Distortion Parameters
+	const double* DistortionParams = DistCoeffs.ptr<double>();
+	const double K1 = DistortionParams[0];
+	const double K2 = DistortionParams[1];
+	const double P1 = DistortionParams[2];
+	const double P2 = DistortionParams[3];
+	const double K3 = DistortionParams[4];
+
+	const int NumPoints = ObjectPoints.total();
+
+	for (int PointIndex = 0; PointIndex < NumPoints; PointIndex++)
+	{
+		const cv::Point3d& WorldPoint = ObjectPoints.at<cv::Point3d>(PointIndex);
+		cv::Point2d& ProjectedPoint = ProjectedPoints.at<cv::Point2d>(PointIndex);
+
+		cv::Point3d CameraPoint;
+		CameraPoint.x = R[0] * WorldPoint.x + R[1] * WorldPoint.y + R[2] * WorldPoint.z + TranslationVector[0];
+		CameraPoint.y = R[3] * WorldPoint.x + R[4] * WorldPoint.y + R[5] * WorldPoint.z + TranslationVector[1];
+		CameraPoint.z = R[6] * WorldPoint.x + R[7] * WorldPoint.y + R[8] * WorldPoint.z + TranslationVector[2];
+
+		// Divide x and y by the z-coordinate to normalize the point before projection
+		CameraPoint.z = CameraPoint.z ? (1.0 / CameraPoint.z) : 1.0;
+
+		CameraPoint.x *= CameraPoint.z;
+		CameraPoint.y *= CameraPoint.z;
+
+		const double R2 = CameraPoint.x * CameraPoint.x + CameraPoint.y * CameraPoint.y;
+		const double R4 = R2 * R2;
+		const double R6 = R4 * R2;
+
+		const double A1 = 2 * CameraPoint.x * CameraPoint.y;
+		const double A2 = R2 + 2 * CameraPoint.x * CameraPoint.x;
+		const double A3 = R2 + 2 * CameraPoint.y * CameraPoint.y;
+
+		const double Radial = 1 + K1 * R2 + K2 * R4 + K3 * R6;
+		const double TangentialX = (P1 * A1) + (P2 * A2);
+		const double TangentialY = (P1 * A3) + (P2 * A1);
+
+		cv::Point2d DistortedPoint;
+		DistortedPoint.x = (CameraPoint.x * Radial) + TangentialX;
+		DistortedPoint.y = (CameraPoint.y * Radial) + TangentialY;
+
+		ProjectedPoint.x = DistortedPoint.x * Fx + Cx;
+		ProjectedPoint.y = DistortedPoint.y * Fy + Cy;
+
+		if (bComputeJacobian)
+		{
+			const int JacobianStep = Jacobian.cols;
+
+			cv::Mat JacRotation;
+			cv::Mat JacTranslation;
+			cv::Mat JacFocalLength;
+			cv::Mat JacImageCenter;
+			cv::Mat JacDistortion;
+			SubdivideJacobian(Jacobian, JacRotation, JacTranslation, JacFocalLength, JacImageCenter, JacDistortion, SolverFlags);
+
+			if (!JacImageCenter.empty())
+			{
+				double* JacImageCenterPtr = JacImageCenter.ptr<double>(PointIndex * 2);
+
+				JacImageCenterPtr[0] = 1;
+				JacImageCenterPtr[1] = 0;
+
+				JacImageCenterPtr[JacobianStep + 0] = 0;
+				JacImageCenterPtr[JacobianStep + 1] = 1;
+			}
+
+			if (!JacFocalLength.empty())
+			{
+				double* JacFocalLengthPtr = JacFocalLength.ptr<double>(PointIndex * 2);
+
+				JacFocalLengthPtr[0] = DistortedPoint.x;
+				JacFocalLengthPtr[1] = 0;
+
+				JacFocalLengthPtr[JacobianStep + 0] = 0;
+				JacFocalLengthPtr[JacobianStep + 1] = DistortedPoint.y;
+			}
+
+			if (!JacDistortion.empty())
+			{
+				double* JacDistortionPtr = JacDistortion.ptr<double>(PointIndex * 2);
+
+				JacDistortionPtr[0] = Fx * CameraPoint.x * R2;
+				JacDistortionPtr[1] = Fx * CameraPoint.x * R4;
+				JacDistortionPtr[2] = Fx * A1;
+				JacDistortionPtr[3] = Fx * A2;
+				JacDistortionPtr[4] = Fx * CameraPoint.x * R6;
+
+				JacDistortionPtr[JacobianStep + 0] = Fy * CameraPoint.y * R2;
+				JacDistortionPtr[JacobianStep + 1] = Fy * CameraPoint.y * R4;
+				JacDistortionPtr[JacobianStep + 2] = Fy * A3;
+				JacDistortionPtr[JacobianStep + 3] = Fy * A1;
+				JacDistortionPtr[JacobianStep + 4] = Fy * CameraPoint.y * R6;
+			}
+
+			if (!JacTranslation.empty())
+			{
+				// Compute the derivative of the undistorted image point with respect to the translation vector [ T0  T1  T2 ]
+				cv::Point3d DCameraPointXDT;
+				DCameraPointXDT.x = CameraPoint.z;
+				DCameraPointXDT.y = 0;
+				DCameraPointXDT.z = -CameraPoint.x * CameraPoint.z;
+
+				cv::Point3d DCameraPointYDT;
+				DCameraPointYDT.x = 0;
+				DCameraPointYDT.y = CameraPoint.z;
+				DCameraPointYDT.z = -CameraPoint.y * CameraPoint.z;
+
+				// Compute the derivative of the R-Squared term with respect to the translation vector [ T0  T1  T2 ]
+				const cv::Point3d DR2DT = (2 * CameraPoint.x * DCameraPointXDT) + (2 * CameraPoint.y * DCameraPointYDT);
+
+				// Compute the derivative of the radial distortion term with respect to the translation vector [ T0  T1  T2 ]
+				const cv::Point3d DRadialDT = (K1 * DR2DT) + (2 * K2 * R2 * DR2DT) + (3 * K3 * R4 * DR2DT);
+
+				// Compute the derivative of the A1 term with respect to the translation vector [ T0  T1  T2 ]
+				const cv::Point3d DA1DT = 2 * (CameraPoint.x * DCameraPointYDT + CameraPoint.y * DCameraPointXDT);
+
+				// Compute the derivative of distorted point with respect to the translation vector [ T0  T1  T2 ]
+				const cv::Point3d DDistortedPointXDT = (DCameraPointXDT * Radial) + (CameraPoint.x * DRadialDT) + (P1 * DA1DT) + (P2 * (DR2DT + 4 * CameraPoint.x * DCameraPointXDT));
+				const cv::Point3d DDistortedPointYDT = (DCameraPointYDT * Radial) + (CameraPoint.y * DRadialDT) + (P2 * DA1DT) + (P1 * (DR2DT + 4 * CameraPoint.y * DCameraPointYDT));
+
+				// Add derivatives to the Jacobian matrix
+				double* JacTranslationPtr = JacTranslation.ptr<double>(PointIndex * 2);
+
+				JacTranslationPtr[0] = Fx * DDistortedPointXDT.x;
+				JacTranslationPtr[1] = Fx * DDistortedPointXDT.y;
+				JacTranslationPtr[2] = Fx * DDistortedPointXDT.z;
+
+				JacTranslationPtr[JacobianStep + 0] = Fy * DDistortedPointYDT.x;
+				JacTranslationPtr[JacobianStep + 1] = Fy * DDistortedPointYDT.y;
+				JacTranslationPtr[JacobianStep + 2] = Fy * DDistortedPointYDT.z;
+
+			}
+
+			if (!JacRotation.empty())
+			{
+				// Compute the derivatives of the camera points with respect to the rotation vector [ R0  R1  R2]
+				// This requires the use of the rotation jacobian matrix that relates the rotation vector in Rodrigues form to the rotation matrix
+				cv::Point3d DCameraPointXDR;
+				DCameraPointXDR.x = WorldPoint.x * JacR[0] + WorldPoint.y * JacR[1] + WorldPoint.z * JacR[2];
+				DCameraPointXDR.y = WorldPoint.x * JacR[9] + WorldPoint.y * JacR[10] + WorldPoint.z * JacR[11];
+				DCameraPointXDR.z = WorldPoint.x * JacR[18] + WorldPoint.y * JacR[19] + WorldPoint.z * JacR[20];
+
+				cv::Point3d DCameraPointYDR;
+				DCameraPointYDR.x = WorldPoint.x * JacR[3] + WorldPoint.y * JacR[4] + WorldPoint.z * JacR[5];
+				DCameraPointYDR.y = WorldPoint.x * JacR[12] + WorldPoint.y * JacR[13] + WorldPoint.z * JacR[14];
+				DCameraPointYDR.z = WorldPoint.x * JacR[21] + WorldPoint.y * JacR[22] + WorldPoint.z * JacR[23];
+
+				cv::Point3d DCameraPointZDR;
+				DCameraPointZDR.x = WorldPoint.x * JacR[6] + WorldPoint.y * JacR[7] + WorldPoint.z * JacR[8];
+				DCameraPointZDR.y = WorldPoint.x * JacR[15] + WorldPoint.y * JacR[16] + WorldPoint.z * JacR[17];
+				DCameraPointZDR.z = WorldPoint.x * JacR[24] + WorldPoint.y * JacR[25] + WorldPoint.z * JacR[26];
+
+				// Compute the derivative of the undistorted image point with respect to the rotation vector [ R0  R1  R2 ]
+				const cv::Point3d DImagePointXDR = CameraPoint.z * (DCameraPointXDR - CameraPoint.x * DCameraPointZDR);
+				const cv::Point3d DImagePointYDR = CameraPoint.z * (DCameraPointYDR - CameraPoint.y * DCameraPointZDR);
+
+				// Compute the derivative of the R-Squared term with respect to the rotation vector [ R0  R1  R2 ]
+				const cv::Point3d DR2DR = 2 * ((CameraPoint.x * DImagePointXDR) + (CameraPoint.y * DImagePointYDR));
+
+				// Compute the derivative of the radial distortion term with respect to the rotation vector [ R0  R1  R2 ]
+				const cv::Point3d DRadialDR = DR2DR * (K1 + (2 * K2 * R2) + (3 * K3 * R4));
+
+				// Compute the derivative of the A1 term with respect to the rotation vector [ R0  R1  R2 ]
+				const cv::Point3d DA1DR = 2 * ((CameraPoint.x * DImagePointYDR) + (CameraPoint.y * DImagePointXDR));
+				
+				// Compute the derivative of distorted point with respect to the rotation vector [ R0  R1  R2 ]
+				const cv::Point3d DDistortedPointXDR = (DImagePointXDR * Radial) + (CameraPoint.x * DRadialDR) + (P1 * DA1DR) + (P2 * (DR2DR + 4 * CameraPoint.x * DImagePointXDR));
+				const cv::Point3d DDistortedPointYDR = (DImagePointYDR * Radial) + (CameraPoint.y * DRadialDR) + (P2 * DA1DR) + (P1 * (DR2DR + 4 * CameraPoint.y * DImagePointYDR));
+
+				// Add derivatives to the Jacobian matrix
+				double* JacRotationPtr = JacRotation.ptr<double>(PointIndex * 2);
+
+				JacRotationPtr[0] = Fx * DDistortedPointXDR.x;
+				JacRotationPtr[1] = Fx * DDistortedPointXDR.y;
+				JacRotationPtr[2] = Fx * DDistortedPointXDR.z;
+
+				JacRotationPtr[JacobianStep + 0] = Fy * DDistortedPointYDR.x;
+				JacRotationPtr[JacobianStep + 1] = Fy * DDistortedPointYDR.y;
+				JacRotationPtr[JacobianStep + 2] = Fy * DDistortedPointYDR.z;
+
+			}
+		}
+	}
+}
+
+void FCameraCalibrationSolver::GatherPoints(
+	const TArray<TArray<FVector>>& ObjectPoints,
+	const TArray<TArray<FVector2D>>& ImagePoints,
+	cv::Mat& ObjectPointsMat,
+	cv::Mat& ImagePointsMat)
+{
+	cv::Point3d* ObjectPointsMatData = ObjectPointsMat.ptr<cv::Point3d>();
+	cv::Point2d* ImagePointsMatData = ImagePointsMat.ptr<cv::Point2d>();
+
+	int TotalPointIndex = 0;
+
+	const int NumImages = ObjectPoints.Num();
+
+	for (int ImageIndex = 0; ImageIndex < NumImages; ++ImageIndex)
+	{
+		const TArray<FVector>& CurrentObjectPoints = ObjectPoints[ImageIndex];
+		const TArray<FVector2D>& CurrentImagePoints = ImagePoints[ImageIndex];
+
+		const int NumPointsInImage = CurrentObjectPoints.Num();
+
+		for (int PointIndex = 0; PointIndex < NumPointsInImage; ++PointIndex)
+		{
+			ObjectPointsMatData[TotalPointIndex + PointIndex] = cv::Point3d(CurrentObjectPoints[PointIndex].X, CurrentObjectPoints[PointIndex].Y, CurrentObjectPoints[PointIndex].Z);
+			ImagePointsMatData[TotalPointIndex + PointIndex] = cv::Point2d(CurrentImagePoints[PointIndex].X, CurrentImagePoints[PointIndex].Y);
+		}
+
+		TotalPointIndex += NumPointsInImage;
+	}
+}
+
+void FCameraCalibrationSolver::SubdivideJacobian(
+	const cv::Mat& Jacobian,
+	cv::Mat& JacRotation,
+	cv::Mat& JacTranslation,
+	cv::Mat& JacFocalLength,
+	cv::Mat& JacImageCenter,
+	cv::Mat& JacDistortion,
+	int SolverFlags)
+{
+	const bool bComputeJacobian = !Jacobian.empty();
+
+	if (bComputeJacobian)
+	{
+		if (Jacobian.cols >= 6)
+		{
+			JacRotation = cv::Mat(Jacobian.colRange(0, 3));
+			JacTranslation = cv::Mat(Jacobian.colRange(3, 6));
+		}
+
+		if (Jacobian.cols >= 10)
+		{
+			if (!(SolverFlags & cv::CALIB_FIX_FOCAL_LENGTH))
+			{
+				JacFocalLength = cv::Mat(Jacobian.colRange(6, 8));
+			}
+
+			if (!(SolverFlags & cv::CALIB_FIX_PRINCIPAL_POINT))
+			{
+				JacImageCenter = cv::Mat(Jacobian.colRange(8, 10));
+			}
+		}
+
+		if (Jacobian.cols > 10)
+		{
+			JacDistortion = cv::Mat(Jacobian.colRange(10, Jacobian.cols));
+		}
+	}
+}
+
+FLevMarqSolver::FLevMarqSolver(int NumParamsToSolve, int NumErrors, int NumMaxIterations, ESymmetryMode SymmetryMode)
+{
+	Params = cv::Mat(NumParamsToSolve, 1, CV_64F);
+	PreviousParams = cv::Mat(NumParamsToSolve, 1, CV_64F);
+	Mask = cv::Mat(NumParamsToSolve, 1, CV_8U, cv::Scalar(1));
+
+	JacTJac = cv::Mat::zeros(NumParamsToSolve, NumParamsToSolve, CV_64F);
+	JacTErr = cv::Mat::zeros(NumParamsToSolve, 1, CV_64F);
+
+	if (NumErrors > 0)
+	{
+		Jacobian = cv::Mat::zeros(NumErrors, NumParamsToSolve, CV_64F);
+		Error = cv::Mat::zeros(NumErrors, 1, CV_64F);
+	}
+
+	MaxIterations = FMath::Clamp(NumMaxIterations, 1, 1000);
+
+	SymmMode = SymmetryMode;
+}
+
+bool FLevMarqSolver::Update()
+{
+	if (State == ESolverState::Done)
+	{
+		return false;
+	}
+
+	if (State == ESolverState::Started)
+	{
+		State = ESolverState::ComputeJacobian;
+		return true;
+	}
+
+	if (State == ESolverState::ComputeJacobian)
+	{
+		JacTJac = Jacobian.t() * Jacobian;
+		JacTErr = Jacobian.t() * Error;
+
+		Params.copyTo(PreviousParams);
+
+		Step();
+
+		if (Iterations == 0)
+		{
+			PreviousErrorNorm = cv::norm(Error, cv::NORM_L2);
+		}
+
+		Error.setTo(cv::Scalar::all(0));
+
+		State = ESolverState::CheckError;
+		return true;
+	}
+
+	ErrorNorm = cv::norm(Error, cv::NORM_L2);
+	if (ErrorNorm > PreviousErrorNorm)
+	{
+		if (++LambdaLog10 <= 16)
+		{
+			Step();
+			Error.setTo(cv::Scalar::all(0));
+			State = ESolverState::CheckError;
+			return true;
+		}
+	}
+
+	LambdaLog10 = MAX(LambdaLog10 - 1, -16);
+	if (++Iterations >= MaxIterations || cv::norm(Params, PreviousParams, cv::NORM_L2 | cv::NORM_RELATIVE) < DBL_EPSILON)
+	{
+		State = ESolverState::Done;
+		return false;
+	}
+
+	PreviousErrorNorm = ErrorNorm;
+	Jacobian.setTo(cv::Scalar::all(0));
+	State = ESolverState::ComputeJacobian;
+	return true;
+}
+
+bool FLevMarqSolver::UpdateAlt()
+{
+	if (State == ESolverState::Done)
+	{
+		return false;
+	}
+
+	if (State == ESolverState::Started)
+	{
+		State = ESolverState::ComputeJacobian;
+		return true;
+	}
+
+	if (State == ESolverState::ComputeJacobian)
+	{
+		Params.copyTo(PreviousParams);
+		Step();
+		PreviousErrorNorm = ErrorNorm;
+		ErrorNorm = 0;
+		State = ESolverState::CheckError;
+		return true;
+	}
+
+	if (ErrorNorm > PreviousErrorNorm)
+	{
+		if (++LambdaLog10 <= 16)
+		{
+			Step();
+			ErrorNorm = 0;
+			State = ESolverState::CheckError;
+			return true;
+		}
+	}
+
+	LambdaLog10 = MAX(LambdaLog10 - 1, -16);
+	if (++Iterations >= MaxIterations || cv::norm(Params, PreviousParams, cv::NORM_L2 | cv::NORM_RELATIVE) < DBL_EPSILON)
+	{
+		State = ESolverState::Done;
+		return false;
+	}
+
+	PreviousErrorNorm = ErrorNorm;
+	JacTJac.setTo(cv::Scalar::all(0));
+	JacTErr.setTo(cv::Scalar::all(0));
+	State = ESolverState::ComputeJacobian;
+	return true;
+}
+
+void FLevMarqSolver::SubMatrix(const cv::Mat& Src, cv::Mat& Dst, const cv::Mat& Cols, const cv::Mat& Rows)
+{
+	int NumNonZerosCols = cv::countNonZero(Cols);
+	cv::Mat tmp(Src.rows, NumNonZerosCols, CV_64FC1);
+
+	for (int i = 0, Index = 0; i < Cols.total(); i++)
+	{
+		if (Cols.at<uchar>(i))
+		{
+			Src.col(i).copyTo(tmp.col(Index++));
+		}
+	}
+
+	int NumNonZeroRows = cv::countNonZero(Rows);
+	Dst.create(NumNonZeroRows, NumNonZerosCols, CV_64FC1);
+	for (int i = 0, Index = 0; i < Rows.total(); i++)
+	{
+		if (Rows.at<uchar>(i))
+		{
+			tmp.row(i).copyTo(Dst.row(Index++));
+		}
+	}
+}
+
+void FLevMarqSolver::Step()
+{
+	cv::Mat SubJtJ;
+	SubMatrix(JacTJac, SubJtJ, Mask, Mask);
+
+	cv::Mat SubJtErr;
+	SubMatrix(JacTErr, SubJtErr, cv::Mat::ones(1, 1, CV_8U), Mask);
+
+	bool bLowerToUpper = !(SymmMode == ESymmetryMode::LowerToUpper);
+	cv::completeSymm(SubJtJ, bLowerToUpper);
+
+	const double Lambda = FMath::Exp(LambdaLog10 * FMath::Loge(10.0));
+	SubJtJ.diag() *= 1.0 + Lambda;
+
+	cv::Mat NonZeroParams;
+	cv::solve(SubJtJ, SubJtErr, NonZeroParams, cv::DECOMP_SVD);
+
+	int NonZeroParamIndex = 0;
+	int NumParamsToSolve = Params.rows;
+
+	for (int ParamIndex = 0; ParamIndex < NumParamsToSolve; ++ParamIndex)
+	{
+		bool bNonZero = (Mask.at<uchar>(ParamIndex) == 1);
+		Params.at<double>(ParamIndex) = PreviousParams.at<double>(ParamIndex) - (bNonZero ? NonZeroParams.at<double>(NonZeroParamIndex++) : 0);
+	}
+}
+
+#endif // WITH_OPENCV
diff --git a/Engine/Plugins/VirtualProduction/CameraCalibration/Source/CameraCalibrationEditor/Private/Calibrators/CameraCalibrationSolver.h b/Engine/Plugins/VirtualProduction/CameraCalibration/Source/CameraCalibrationEditor/Private/Calibrators/CameraCalibrationSolver.h
new file mode 100644
--- /dev/null
+++ b/Engine/Plugins/VirtualProduction/CameraCalibration/Source/CameraCalibrationEditor/Private/Calibrators/CameraCalibrationSolver.h
@@ -0,0 +1,208 @@
+// Copyright Epic Games, Inc. All Rights Reserved.
+
+#pragma once
+
+#include "Models/LensModel.h"
+#include "Templates/SubclassOf.h"
+
+#if WITH_OPENCV
+
+#include "PreOpenCVHeaders.h"
+#include <opencv2/core.hpp>
+#include "PostOpenCVHeaders.h"
+
+#endif	// WITH_OPENCV
+
+/** 
+  * Lens Distortion Solver class, supporting anamorphic and spherical models 
+  * The implementation is largely based on the implementation of calibrateCamera from OpenCV: [REDACTED_PROFILE_URL]/opencv
+  */
+class FCameraCalibrationSolver
+{
+#if WITH_OPENCV
+public:
+	/** 
+	  * Calibrate camera intrinsics and distortion parameters using a set of input 3D and 2D point correspondences
+	  * Returns the root mean reprojection error between the input image points and the projection of the input object points
+	  */
+	static double CalibrateCamera(
+		const TSubclassOf<ULensModel> LensModel,
+		const TArray<TArray<FVector>>& InObjectPoints,
+		const TArray<TArray<FVector2D>>& InImagePoints,
+		const FIntPoint ImageSize,
+		FVector2f& InOutFocalLength,
+		FVector2f& InOutImageCenter,
+		TArray<float>& OutDistCoeffs,
+		double PixelAspect = 1.0,
+		int SolverFlags = 0);
+
+private:
+	/** Initialize the camera matrix of intrinsic parameters using linear algebra techniques */
+	static void InitCameraIntrinsics(
+		const cv::Mat& ObjectPoints,
+		const cv::Mat& ImagePoints,
+		const cv::Mat& NumPoints,
+		const cv::Size ImageSize,
+		cv::Mat& CameraMatrix);
+
+	/** Initialize the camera extrinsics (rotation and translation vectors) for each image using linear algebra techniques */
+	static void InitCameraExtrinsics(
+		const TSubclassOf<ULensModel> LensModel,
+		const cv::Mat& ObjectPoints,
+		const cv::Mat& ImagePoints,
+		const cv::Mat& CameraMatrix,
+		const cv::Mat& DistCoeffs,
+		const cv::Size ImageSize,
+		cv::Mat& Rotation,
+		cv::Mat& Translation,
+		const int SolverFlags);
+
+	/** Project the input object points to 2D using the input camera intrinsics, extrinsics, and distortion parameters */
+	static void ProjectPoints(
+		const TSubclassOf<ULensModel> LensModel,
+		const cv::Mat& ObjectPoints,
+		const cv::Mat& Rotation,
+		const cv::Mat& Translation,
+		const cv::Mat& CameraMatrix,
+		const cv::Mat& DistCoeffs,
+		const cv::Size ImageSize,
+		cv::Mat& ProjectedPoints);
+
+	/** Project the input object points to 2D using the input camera intrinsics, extrinsics, and distortion parameters */
+	static void ProjectPoints(
+		const TSubclassOf<ULensModel> LensModel,
+		const cv::Mat& ObjectPoints,
+		const cv::Mat& Rotation,
+		const cv::Mat& Translation,
+		const cv::Mat& CameraMatrix,
+		const cv::Mat& DistCoeffs,
+		const cv::Size ImageSize,
+		cv::Mat& ProjectedPoints,
+		cv::Mat& Jacobian,
+		int SolverFlags);
+
+	/** Project the input object points to 2D using the input camera intrinsics, extrinsics, and anamorphic distortion parameters */
+	static void ProjectPointsAnamorphic(
+		const cv::Mat& ObjectPoints,
+		const cv::Mat& Rotation,
+		const cv::Mat& Translation,
+		const cv::Mat& CameraMatrix,
+		const cv::Mat& DistCoeffs,
+		const cv::Size ImageSize,
+		cv::Mat& ProjectedPoints,
+		cv::Mat& Jacobian,
+		int SolverFlags);
+
+	/** Project the input object points to 2D using the input camera intrinsics, extrinsics, and spherical distortion parameters */
+	static void ProjectPointsSpherical(
+		const cv::Mat& ObjectPoints,
+		const cv::Mat& Rotation,
+		const cv::Mat& Translation,
+		const cv::Mat& CameraMatrix,
+		const cv::Mat& DistCoeffs,
+		cv::Mat& ProjectedPoints,
+		cv::Mat& Jacobian,
+		int SolverFlags);
+
+	/** Copy the input 3D and 2D points to OpenCV matrices for ease of use with the solver */
+	static void GatherPoints(
+		const TArray<TArray<FVector>>& ObjectPoints,
+		const TArray<TArray<FVector2D>>& ImagePoints,
+		cv::Mat& ObjectPointsMat,
+		cv::Mat& ImagePointsMat);
+
+	/** Divide large jacobian matrix into smaller views for each of the parameter groups the solver will solve */
+	static void SubdivideJacobian(
+		const cv::Mat& Jacobian,
+		cv::Mat& JacRotation,
+		cv::Mat& JacTranslation,
+		cv::Mat& JacFocalLength,
+		cv::Mat& JacImageCenter,
+		cv::Mat& JacDistortion,
+		int SolverFlags);
+
+#endif	// WITH_OPENCV
+};
+
+#if WITH_OPENCV
+
+/** Levenberg–Marquardt solver for solving least-squares problems */
+class FLevMarqSolver
+{
+public:
+	/** Mode for copying the contents of symmetric matrix */
+	enum class ESymmetryMode : uint8
+	{
+		LowerToUpper,
+		UpperToLower
+	};
+
+	/** State used by the solver in its state machine */
+	enum class ESolverState : uint8
+	{
+		Done = 0,
+		Started = 1,
+		ComputeJacobian = 2,
+		CheckError = 3
+	};
+
+	FLevMarqSolver(int NumParamsToSolve, int NumErrors, int NumMaxIterations, ESymmetryMode SymmetryMode = ESymmetryMode::LowerToUpper);
+
+	/** Update the solver with the input jacobian and error matrices */
+	bool Update();
+
+	/** Update the solver with the input pre-multiplied jacobian and error matrices */
+	bool UpdateAlt();
+
+private:
+	/** Step the solver, modifying the values of each of the parameters */
+	void Step();
+
+	/** Create a submatrix from the source matrix using only the columns and rows specified */
+	void SubMatrix(const cv::Mat& Src, cv::Mat& Dst, const cv::Mat& Cols, const cv::Mat& Rows);
+
+public:
+	/** Parameters solved by this solver */
+	cv::Mat Params;
+
+	/** Mask specifiying which parameters in the parameter list should not be modified */
+	cv::Mat Mask;
+
+	/** Current solver state */
+	ESolverState State = ESolverState::Started;
+
+	/** Average error for the current parameter set */
+	double ErrorNorm = 0;
+
+	/** Jacobian Matrix */
+	cv::Mat Jacobian;
+
+	/** Error Matrix */
+	cv::Mat Error;
+
+	/** Jacobian Transposed multiplied by the Jacobian */
+	cv::Mat JacTJac;
+
+	/** Jacobian Transposed multiplied by the Error */
+	cv::Mat JacTErr;
+
+private:
+	/** Stores the previous values of each parameter */
+	cv::Mat PreviousParams;
+
+	/** Previous average error */
+	double PreviousErrorNorm = DBL_MAX;
+
+	/** Maximum number of iterations the solver will do */
+	int MaxIterations = 0;
+
+	/** Current number of iterations */
+	int Iterations = 0;
+
+	/** Mode to use when copying half of the symmetric matrix */
+	ESymmetryMode SymmMode;
+
+	int LambdaLog10 = -3;
+};
+
+#endif // WITH_OPENCV
diff --git a/Engine/Plugins/VirtualProduction/CameraCalibration/Source/CameraCalibrationEditor/Private/Calibrators/CameraLensDistortionAlgoCheckerboard.cpp b/Engine/Plugins/VirtualProduction/CameraCalibration/Source/CameraCalibrationEditor/Private/Calibrators/CameraLensDistortionAlgoCheckerboard.cpp
--- a/Engine/Plugins/VirtualProduction/CameraCalibration/Source/CameraCalibrationEditor/Private/Calibrators/CameraLensDistortionAlgoCheckerboard.cpp
+++ b/Engine/Plugins/VirtualProduction/CameraCalibration/Source/CameraCalibrationEditor/Private/Calibrators/CameraLensDistortionAlgoCheckerboard.cpp
@@ -6,8 +6,10 @@
 #include "AssetEditor/LensDistortionTool.h"
 #include "AssetEditor/SSimulcamViewport.h"
 #include "AssetRegistry/AssetData.h"
+#include "Camera/CameraActor.h"
 #include "CameraCalibrationCheckerboard.h"
 #include "CameraCalibrationEditorLog.h"
+#include "CameraCalibrationSolver.h"
 #include "Dom/JsonObject.h"
 #include "EditorFontGlyphs.h"
 #include "Engine/Texture2D.h"
@@ -15,17 +17,26 @@
 #include "ImageUtils.h"
 #include "JsonObjectConverter.h"
 #include "Misc/MessageDialog.h"
+#include "Models/AnamorphicLensModel.h"
 #include "PropertyCustomizationHelpers.h"
 #include "SphericalLensDistortionModelHandler.h"
 #include "UI/CameraCalibrationWidgetHelpers.h"
 #include "Widgets/Input/SButton.h"
 #include "Widgets/Input/SCheckBox.h"
+#include "Widgets/Input/SNumericEntryBox.h"
 #include "Widgets/Views/SListView.h"
 #include "Widgets/Views/STableRow.h"
 
 
 #define LOCTEXT_NAMESPACE "CameraLensDistortionAlgoCheckerboard"
 
+#if WITH_EDITOR
+static TAutoConsoleVariable<float> CVarFixFocalLengthValue(TEXT("CameraCalibration.FocalLengthValue"), 0.0f, TEXT("Fixed focal length value in mm"));
+static TAutoConsoleVariable<float> CVarSensorWidthValue(TEXT("CameraCalibration.SensorWidthValue"), 0.0f, TEXT("Fixed focal length value in mm"));
+static TAutoConsoleVariable<float> CVarImageWidthValue(TEXT("CameraCalibration.ImageWidthValue"), 0, TEXT("Fixed focal length value in pixels"));
+static TAutoConsoleVariable<float> CVarImageHeightValue(TEXT("CameraCalibration.ImageHeightValue"), 0, TEXT("Fixed focal length value in pixels"));
+#endif
+
 const int UCameraLensDistortionAlgoCheckerboard::DATASET_VERSION = 1;
 
 // String constants for import/export
@@ -139,7 +150,7 @@ void UCameraLensDistortionAlgoCheckerboard::Shutdown()
 
 bool UCameraLensDistortionAlgoCheckerboard::SupportsModel(const TSubclassOf<ULensModel>& LensModel) const
 {
-	return (LensModel == USphericalLensModel::StaticClass());
+	return ((LensModel == USphericalLensModel::StaticClass()) || (LensModel == UAnamorphicLensModel::StaticClass()));
 }
 
 void UCameraLensDistortionAlgoCheckerboard::Tick(float DeltaTime)
@@ -423,30 +434,48 @@ TSharedRef<SWidget> UCameraLensDistortionAlgoCheckerboard::BuildUI()
 		.VAlign(EVerticalAlignment::VAlign_Top)
 		.AutoHeight()
 		.MaxHeight(FCameraCalibrationWidgetHelpers::DefaultRowHeight)
-		[ FCameraCalibrationWidgetHelpers::BuildLabelWidgetPair(LOCTEXT("Checkerboard", "Checkerboard"), BuildCalibrationDevicePickerWidget()) ]
+		[FCameraCalibrationWidgetHelpers::BuildLabelWidgetPair(LOCTEXT("Checkerboard", "Checkerboard"), BuildCalibrationDevicePickerWidget())]
 
-		+ SVerticalBox::Slot() // Show Overlay
+	+ SVerticalBox::Slot() // Show Overlay
 		.VAlign(EVerticalAlignment::VAlign_Top)
 		.AutoHeight()
 		.MaxHeight(FCameraCalibrationWidgetHelpers::DefaultRowHeight)
 		[FCameraCalibrationWidgetHelpers::BuildLabelWidgetPair(LOCTEXT("ShowOverlay", "Show Coverage Overlay"), BuildShowOverlayWidget())]
 
-		+ SVerticalBox::Slot() // Show Detection
+	+ SVerticalBox::Slot() // Show Detection
 		.VAlign(EVerticalAlignment::VAlign_Top)
 		.AutoHeight()
 		.MaxHeight(FCameraCalibrationWidgetHelpers::DefaultRowHeight)
 		[FCameraCalibrationWidgetHelpers::BuildLabelWidgetPair(LOCTEXT("ShowDetection", "Show Detection"), BuildShowDetectionWidget())]
 
-		+ SVerticalBox::Slot() // Calibration Rows
+	+ SVerticalBox::Slot() // Pixel Aspect
+		.VAlign(EVerticalAlignment::VAlign_Top)
+		.AutoHeight()
+		.MaxHeight(FCameraCalibrationWidgetHelpers::DefaultRowHeight)
+		[FCameraCalibrationWidgetHelpers::BuildLabelWidgetPair(LOCTEXT("PixelAspect", "Pixel Aspect"), BuildPixelAspectWidget())]
+
+	+ SVerticalBox::Slot() // Fix Focal Length
+		.VAlign(EVerticalAlignment::VAlign_Top)
+		.AutoHeight()
+		.MaxHeight(FCameraCalibrationWidgetHelpers::DefaultRowHeight)
+		[FCameraCalibrationWidgetHelpers::BuildLabelWidgetPair(LOCTEXT("FixFocalLength", "Fix Focal Length"), BuildFixFocalLengthWidget())]
+
+	+ SVerticalBox::Slot() // Fix Image Center
+		.VAlign(EVerticalAlignment::VAlign_Top)
+		.AutoHeight()
+		.MaxHeight(FCameraCalibrationWidgetHelpers::DefaultRowHeight)
+		[FCameraCalibrationWidgetHelpers::BuildLabelWidgetPair(LOCTEXT("FixImageCenter", "Fix Image Center"), BuildFixImageCenterWidget())]
+
+	+ SVerticalBox::Slot() // Calibration Rows
 		.AutoHeight()
 		.MaxHeight(12 * FCameraCalibrationWidgetHelpers::DefaultRowHeight)
-		[ BuildCalibrationPointsTable() ]
-		
-		+ SVerticalBox::Slot() // Action buttons (e.g. Remove, Clear)
+		[BuildCalibrationPointsTable()]
+
+	+ SVerticalBox::Slot() // Action buttons (e.g. Remove, Clear)
 		.HAlign(EHorizontalAlignment::HAlign_Center)
 		.AutoHeight()
-		.Padding(0,20)
-		[ BuildCalibrationActionButtons() ];
+		.Padding(0, 20)
+		[BuildCalibrationActionButtons()];
 }
 
 bool UCameraLensDistortionAlgoCheckerboard::ValidateNewRow(TSharedPtr<FLensDistortionCheckerboardRowData>& Row, FText& OutErrorMessage) const
@@ -468,7 +497,7 @@ bool UCameraLensDistortionAlgoCheckerboard::ValidateNewRow(TSharedPtr<FLensDisto
 		OutErrorMessage = LOCTEXT("InvalidCameraData", "Invalid CameraData");
 		return false;
 	}
-	
+
 	// Valid image dimensions
 	if ((Row->ImageHeight < 1) || (Row->ImageWidth < 1))
 	{
@@ -561,9 +590,9 @@ bool UCameraLensDistortionAlgoCheckerboard::GetLensDistortion(
 	//
 
 	// Enough points
-	if (CalibrationRows.Num() < 4)
+	if (CalibrationRows.Num() < 1)
 	{
-		OutErrorMessage = LOCTEXT("NotEnoughSamples", "At least 4 calibration rows are required");
+		OutErrorMessage = LOCTEXT("NotEnoughSamples", "At least 1 calibration row is required");
 		return false;
 	}
 
@@ -605,105 +634,95 @@ bool UCameraLensDistortionAlgoCheckerboard::GetLensDistortion(
 		return false;
 	}
 
-	// Only spherical lens distortion is currently supported at the moment.
+#if WITH_OPENCV
+
+	// Initialize the camera matrix that will be used in each call to projectPoints()
+	float TestImageWidth = LensFile->CameraFeedInfo.GetDimensions().X;
+	float TestImageHeight = LensFile->CameraFeedInfo.GetDimensions().Y;
 
-	const USphericalLensDistortionModelHandler* SphericalHandler = Cast<USphericalLensDistortionModelHandler>(StepsController->GetDistortionHandler());
+	float TestFocalLength = StepsController->GetLensFileEvaluationInputs().Zoom;
+	float TestSensorWidth = StepsController->GetLensFileEvaluationInputs().Filmback.SensorWidth;
 
-	if (!SphericalHandler)
+	if (CVarFixFocalLengthValue.GetValueOnGameThread() > 1.0f)
 	{
-		OutErrorMessage = LOCTEXT("OnlySphericalDistortionSupported", "Only spherical distortion is currently supported. Please update the distortion model used by the camera.");
-		return false;
+		TestFocalLength = CVarFixFocalLengthValue.GetValueOnGameThread();
+	}
+	if (CVarSensorWidthValue.GetValueOnGameThread() > 1.0f)
+	{
+		TestSensorWidth = CVarSensorWidthValue.GetValueOnGameThread();
+	}
+	if (CVarImageWidthValue.GetValueOnGameThread() > 1)
+	{
+		TestImageWidth = CVarImageWidthValue.GetValueOnGameThread();
+	}
+	if (CVarImageHeightValue.GetValueOnGameThread() > 1)
+	{
+		TestImageHeight = CVarImageHeightValue.GetValueOnGameThread();
 	}
 
-#if WITH_OPENCV
-
-	cv::Mat CameraMatrix = cv::Mat::eye(3, 3, CV_64F);;
-	cv::Mat DistortionCoefficients;
+	float Fx = TestImageWidth;
+	if (!FMath::IsNearlyZero(TestSensorWidth))
+	{
+		Fx = (TestFocalLength / TestSensorWidth) * TestImageWidth;
+	}
 
-	std::vector<cv::Mat> Rvecs;
-	std::vector<cv::Mat> Tvecs;
+	FVector2f FocalLength = FVector2f(Fx, Fx);
+	FVector2f ImageCenter = FVector2f(TestImageWidth * 0.5f, TestImageHeight * 0.5f);
 
-	std::vector<std::vector<cv::Point2f>> Samples2d;
-	std::vector<std::vector<cv::Point3f>> Samples3d;
+	TArray<TArray<FVector2D>> Samples2d;
+	TArray<TArray<FVector>> Samples3d;
 
 	for (const TSharedPtr<FLensDistortionCheckerboardRowData>& Row : CalibrationRows)
 	{
-		// add 2d (image) points
-		{
-			std::vector<cv::Point2f> Points2d;
-
-			for (FVector2D& Point2d : Row->Points2d)
-			{
-				Points2d.push_back(cv::Point2f(Point2d.X, Point2d.Y));
-			}
-
-			Samples2d.push_back(Points2d);
-		}
-
-		// add 3d points
-		{
-			std::vector<cv::Point3f> Points3d;
+		Samples2d.Add(Row->Points2d);
+		Samples3d.Add(Row->Points3d);
+	}
 
-			for (FVector& Point3d : Row->Points3d)
-			{
-				Points3d.push_back(cv::Point3f(Point3d.X, Point3d.Y, Point3d.Z));
-			}
+	int flags = cv::CALIB_USE_INTRINSIC_GUESS;
+	if (bFixFocalLength)
+	{
+		flags |= cv::CALIB_USE_INTRINSIC_GUESS;
+		flags |= cv::CALIB_FIX_FOCAL_LENGTH;
+	}
 
-			Samples3d.push_back(Points3d);
-		}
+	if (bFixImageCenter)
+	{
+		flags |= cv::CALIB_USE_INTRINSIC_GUESS;
+		flags |= cv::CALIB_FIX_PRINCIPAL_POINT;
 	}
 
-	OutError = cv::calibrateCamera(
+	TArray<float> DistortionCoefficients;
+
+	OutError = FCameraCalibrationSolver::CalibrateCamera(
+		LensFile->LensInfo.LensModel,
 		Samples3d,
 		Samples2d,
-		cv::Size(LastRow->ImageWidth, LastRow->ImageHeight), 
-		CameraMatrix, 
-		DistortionCoefficients, 
-		Rvecs, 
-		Tvecs
+		FIntPoint(LastRow->ImageWidth, LastRow->ImageHeight),
+		FocalLength,
+		ImageCenter,
+		DistortionCoefficients,
+		PixelAspect,
+		flags
 	);
 
-	check(DistortionCoefficients.total() == 5);
-	check((CameraMatrix.rows == 3) && (CameraMatrix.cols == 3));
-
-	// Valid image sizes were verified when adding the calibration rows.
-	checkSlow(LastRow->ImageWidth > 0);
-	checkSlow(LastRow->ImageHeight > 0);
-
-	// FZ inputs to LUT
-	OutFocus = LastRow->CameraData.InputFocus;
-	OutZoom = LastRow->CameraData.InputZoom;
+	OutLensModel = LensFile->LensInfo.LensModel;
+	OutDistortionInfo.Parameters = DistortionCoefficients;
 
 	// FocalLengthInfo
 	OutFocalLengthInfo.FxFy = FVector2D(
-		float(CameraMatrix.at<double>(0, 0) / LastRow->ImageWidth),
-		float(CameraMatrix.at<double>(1, 1) / LastRow->ImageHeight)
+		float(FocalLength.X / LastRow->ImageWidth),
+		float(FocalLength.Y / LastRow->ImageHeight)
 	);
 
-	// DistortionInfo
-	{
-		FSphericalDistortionParameters SphericalParameters;
-
-		SphericalParameters.K1 = DistortionCoefficients.at<double>(0);
-		SphericalParameters.K2 = DistortionCoefficients.at<double>(1);
-		SphericalParameters.P1 = DistortionCoefficients.at<double>(2);
-		SphericalParameters.P2 = DistortionCoefficients.at<double>(3);
-		SphericalParameters.K3 = DistortionCoefficients.at<double>(4);
-
-		USphericalLensModel::StaticClass()->GetDefaultObject<ULensModel>()->ToArray(
-			SphericalParameters, 
-			OutDistortionInfo.Parameters
-		);
-	}
-
 	// ImageCenterInfo
 	OutImageCenterInfo.PrincipalPoint = FVector2D(
-		float(CameraMatrix.at<double>(0, 2) / LastRow->ImageWidth),
-		float(CameraMatrix.at<double>(1, 2) / LastRow->ImageHeight)
+		float(ImageCenter.X / LastRow->ImageWidth),
+		float(ImageCenter.Y / LastRow->ImageHeight)
 	);
 
-	// Lens Model
-	OutLensModel = USphericalLensModel::StaticClass();
+	// FZ inputs to LUT
+	OutFocus = LastRow->CameraData.InputFocus;
+	OutZoom = LastRow->CameraData.InputZoom;
 
 	return true;
 #else
@@ -718,116 +737,151 @@ TSharedRef<SWidget> UCameraLensDistortionAlgoCheckerboard::BuildCalibrationDevic
 {
 	return SNew(SHorizontalBox)
 
-		+ SHorizontalBox::Slot() // Picker
-		[
-			SNew(SObjectPropertyEntryBox)
-			.AllowedClass(ACameraCalibrationCheckerboard::StaticClass())
-			.OnObjectChanged_Lambda([&](const FAssetData& AssetData) -> void
+	+ SHorizontalBox::Slot() // Picker
+	[
+		SNew(SObjectPropertyEntryBox)
+		.AllowedClass(ACameraCalibrationCheckerboard::StaticClass())
+		.OnObjectChanged_Lambda([&](const FAssetData& AssetData) -> void
+		{
+			if (AssetData.IsValid())
 			{
-				if (AssetData.IsValid())
-				{
-					SetCalibrator(Cast<ACameraCalibrationCheckerboard>(AssetData.GetAsset()));
-				}
-			})
-			.ObjectPath_Lambda([&]() -> FString
+				SetCalibrator(Cast<ACameraCalibrationCheckerboard>(AssetData.GetAsset()));
+			}
+		})
+		.ObjectPath_Lambda([&]() -> FString
+		{
+			if (AActor* TheCalibrator = GetCalibrator())
 			{
-				if (AActor* TheCalibrator = GetCalibrator())
-				{
-					FAssetData AssetData(TheCalibrator, true);
-					return AssetData.GetObjectPathString();
-				}
+				FAssetData AssetData(TheCalibrator, true);
+				return AssetData.GetObjectPathString();
+			}
 
-				return TEXT("");
-			})
-		]
+			return TEXT("");
+		})
+	]
+
+	+ SHorizontalBox::Slot() // Spawner
+	.AutoWidth()
+	[
+		SNew(SButton)
+		.Text(LOCTEXT("Spawn", "Spawn"))
+		.HAlign(HAlign_Center)
+		.VAlign(VAlign_Center)
+		.ButtonStyle(FAppStyle::Get(), "HoverHintOnly")
+		.OnClicked_Lambda([&]() -> FReply
+		{
+			const FCameraCalibrationStepsController* StepsController = GetStepsController();
 
-		+ SHorizontalBox::Slot() // Spawner
-		.AutoWidth()
-		[
-			SNew(SButton)
-			.Text(LOCTEXT("Spawn", "Spawn"))
-			.HAlign(HAlign_Center)
-			.VAlign(VAlign_Center)
-			.ButtonStyle(FAppStyle::Get(), "HoverHintOnly")
-			.OnClicked_Lambda([&]() -> FReply
+			if (!ensure(StepsController))
 			{
-				const FCameraCalibrationStepsController* StepsController = GetStepsController();
-
-				if (!ensure(StepsController))
-				{
-					return FReply::Handled();
-				}
+				return FReply::Handled();
+			}
 
-				if (UWorld* const World = StepsController->GetWorld())
-				{
-					SetCalibrator(World->SpawnActor<ACameraCalibrationCheckerboard>());
-				}
+			if (UWorld* const World = StepsController->GetWorld())
+			{
+				SetCalibrator(World->SpawnActor<ACameraCalibrationCheckerboard>());
+			}
 
-				return FReply::Handled();
-			})
-			[
-				SNew(STextBlock)
-				.Font(FAppStyle::Get().GetFontStyle("FontAwesome.12"))
-				.Text(FEditorFontGlyphs::Plus)
-				.ColorAndOpacity(FLinearColor::White)
-			]
+			return FReply::Handled();
+		})
+		[
+			SNew(STextBlock)
+			.Font(FAppStyle::Get().GetFontStyle("FontAwesome.12"))
+			.Text(FEditorFontGlyphs::Plus)
+			.ColorAndOpacity(FLinearColor::White)
 		]
-		;
+	];
 }
 
 TSharedRef<SWidget> UCameraLensDistortionAlgoCheckerboard::BuildShowOverlayWidget()
 {
 	return SNew(SCheckBox)
-		.IsChecked_Lambda([&]() -> ECheckBoxState
-		{
-			return bShouldShowOverlay ? ECheckBoxState::Checked : ECheckBoxState::Unchecked;
-		})
-		.OnCheckStateChanged_Lambda([&](ECheckBoxState NewState) -> void
-		{
-			bShouldShowOverlay = (NewState == ECheckBoxState::Checked);
-
-			FCameraCalibrationStepsController* StepsController = Tool->GetCameraCalibrationStepsController();
+	.IsChecked_Lambda([&]() -> ECheckBoxState
+	{
+		return bShouldShowOverlay ? ECheckBoxState::Checked : ECheckBoxState::Unchecked;
+	})
+	.OnCheckStateChanged_Lambda([&](ECheckBoxState NewState) -> void
+	{
+		bShouldShowOverlay = (NewState == ECheckBoxState::Checked);
 
-			if (!ensure(StepsController))
-			{
-				return;
-			}
+		FCameraCalibrationStepsController* StepsController = Tool->GetCameraCalibrationStepsController();
+		
+		if (!ensure(StepsController))
+		{
+			return;
+		}
 
-			StepsController->SetOverlayEnabled(bShouldShowOverlay);
-		});
+		StepsController->SetOverlayEnabled(bShouldShowOverlay);
+	});
 }
 
 TSharedRef<SWidget> UCameraLensDistortionAlgoCheckerboard::BuildShowDetectionWidget()
 {
 	return SNew(SCheckBox)
-		.IsChecked_Lambda([&]() -> ECheckBoxState
-		{
-			return bShouldShowDetectionWindow ? ECheckBoxState::Checked : ECheckBoxState::Unchecked;
-		})
-		.OnCheckStateChanged_Lambda([&](ECheckBoxState NewState) -> void
+	.IsChecked_Lambda([&]() -> ECheckBoxState
+	{
+		return bShouldShowDetectionWindow ? ECheckBoxState::Checked : ECheckBoxState::Unchecked;
+	})
+	.OnCheckStateChanged_Lambda([&](ECheckBoxState NewState) -> void
+	{
+		bShouldShowDetectionWindow = (NewState == ECheckBoxState::Checked);
+	});
+}
+
+TSharedRef<SWidget> UCameraLensDistortionAlgoCheckerboard::BuildPixelAspectWidget()
+{
+	return SNew(SNumericEntryBox<float>)
+		.Value_Lambda([&]() { return PixelAspect; })
+		.ToolTipText(LOCTEXT("Pixel Aspect", "Pixel Aspect"))
+		.OnValueChanged_Lambda([&](double InValue)
 		{
-			bShouldShowDetectionWindow = (NewState == ECheckBoxState::Checked);
+			PixelAspect = InValue;
 		});
 }
 
+TSharedRef<SWidget> UCameraLensDistortionAlgoCheckerboard::BuildFixFocalLengthWidget()
+{
+	return SNew(SCheckBox)
+	.IsChecked_Lambda([&]() -> ECheckBoxState
+	{
+		return bFixFocalLength ? ECheckBoxState::Checked : ECheckBoxState::Unchecked;
+	})
+	.OnCheckStateChanged_Lambda([&](ECheckBoxState NewState) -> void
+	{
+		bFixFocalLength = (NewState == ECheckBoxState::Checked);
+	});
+}
+
+TSharedRef<SWidget> UCameraLensDistortionAlgoCheckerboard::BuildFixImageCenterWidget()
+{
+	return SNew(SCheckBox)
+	.IsChecked_Lambda([&]() -> ECheckBoxState
+	{
+		return bFixImageCenter ? ECheckBoxState::Checked : ECheckBoxState::Unchecked;
+	})
+	.OnCheckStateChanged_Lambda([&](ECheckBoxState NewState) -> void
+	{
+		bFixImageCenter = (NewState == ECheckBoxState::Checked);
+	});
+}
+
 TSharedRef<SWidget> UCameraLensDistortionAlgoCheckerboard::BuildCalibrationActionButtons()
 {
 	return SNew(SHorizontalBox)
 
-		+ SHorizontalBox::Slot() // Button to clear all rows
-		.AutoWidth()
-		[ 
-			SNew(SButton)
-			.Text(LOCTEXT("ClearAll", "Clear All"))
-			.HAlign(HAlign_Center)
-			.VAlign(VAlign_Center)
-			.OnClicked_Lambda([&]() -> FReply
-			{
-				ClearCalibrationRows();
-				return FReply::Handled();
-			})
-		]
-		;
+	+ SHorizontalBox::Slot() // Button to clear all rows
+	.AutoWidth()
+	[
+		SNew(SButton)
+		.Text(LOCTEXT("ClearAll", "Clear All"))
+		.HAlign(HAlign_Center)
+		.VAlign(VAlign_Center)
+		.OnClicked_Lambda([&]() -> FReply
+		{
+			ClearCalibrationRows();
+			return FReply::Handled();
+		})
+	];
 }
 
 TSharedRef<SWidget> UCameraLensDistortionAlgoCheckerboard::BuildCalibrationPointsTable()
@@ -838,7 +892,7 @@ TSharedRef<SWidget> UCameraLensDistortionAlgoCheckerboard::BuildCalibrationPoint
 		.OnGenerateRow_Lambda([&](TSharedPtr<FLensDistortionCheckerboardRowData> InItem, const TSharedRef<STableViewBase>& OwnerTable) -> TSharedRef<ITableRow>
 		{
 			return SNew(CameraLensDistortionAlgoCheckerboard::SCalibrationRowGenerator, OwnerTable)
-				.CalibrationRowData(InItem);
+			.CalibrationRowData(InItem);
 		})
 		.SelectionMode(ESelectionMode::Multi)
 		.OnKeyDownHandler_Lambda([&](const FGeometry& Geometry, const FKeyEvent& KeyEvent) -> FReply
diff --git a/Engine/Plugins/VirtualProduction/CameraCalibration/Source/CameraCalibrationEditor/Private/Calibrators/CameraLensDistortionAlgoCheckerboard.h b/Engine/Plugins/VirtualProduction/CameraCalibration/Source/CameraCalibrationEditor/Private/Calibrators/CameraLensDistortionAlgoCheckerboard.h
--- a/Engine/Plugins/VirtualProduction/CameraCalibration/Source/CameraCalibrationEditor/Private/Calibrators/CameraLensDistortionAlgoCheckerboard.h
+++ b/Engine/Plugins/VirtualProduction/CameraCalibration/Source/CameraCalibrationEditor/Private/Calibrators/CameraLensDistortionAlgoCheckerboard.h
@@ -186,6 +186,15 @@ private:
 	/** True if a detection window should be shown after every capture */
 	bool bShouldShowDetectionWindow = false;
 
+	/** Pixel Aspect value to pass to the solver */
+	float PixelAspect = 1.0f;
+
+	/** If true, the solver will not solve for focal length, and will just use the input value */
+	bool bFixFocalLength = false;
+
+	/** If true, the solver will not solve for image center, and will just use the input value */
+	bool bFixImageCenter = false;
+
 	/** Texture into which detected chessboard corners from each calibration row are drawn */
 	TObjectPtr<UTexture2D> CoverageTexture;
 
@@ -205,6 +214,15 @@ private:
 	/** Builds the UI for the user to select if they want a corner detection window to be shown after every capture */
 	TSharedRef<SWidget> BuildShowDetectionWidget();
 
+	/** Builds the UI of the pixel aspect widget */
+	TSharedRef<SWidget> BuildPixelAspectWidget();
+
+	/** Builds the UI of the fix focal length width */
+	TSharedRef<SWidget> BuildFixFocalLengthWidget();
+
+	/** Builds the UI of the fix image center widget */
+	TSharedRef<SWidget> BuildFixImageCenterWidget();
+
 	/** Builds the UI of the calibration points table */
 	TSharedRef<SWidget> BuildCalibrationPointsTable();
 
