CameraCalibration: Fix editor crash when opening a LensFile asset in a level with no camera actor.
#rnx

diff --git a/Engine/Plugins/VirtualProduction/CameraCalibration/Source/CameraCalibrationEditor/Private/AssetEditor/CameraCalibrationStepsController.cpp b/Engine/Plugins/VirtualProduction/CameraCalibration/Source/CameraCalibrationEditor/Private/AssetEditor/CameraCalibrationStepsController.cpp
--- a/Engine/Plugins/VirtualProduction/CameraCalibration/Source/CameraCalibrationEditor/Private/AssetEditor/CameraCalibrationStepsController.cpp
+++ b/Engine/Plugins/VirtualProduction/CameraCalibration/Source/CameraCalibrationEditor/Private/AssetEditor/CameraCalibrationStepsController.cpp
@@ -1097,6 +1097,11 @@ void FCameraCalibrationStepsController::SetCamera(ACameraActor* InCamera)
 {
 	Camera = InCamera;
 
+	if (!InCamera)
+	{
+		return;
+	}
+
 	TInlineComponentArray<UCineCameraComponent*> CameraComponents;
 	InCamera->GetComponents(CameraComponents);
 
