Added explicit RenderCore dependencies to fix linker error.

diff --git a/Engine/Plugins/Experimental/ChaosVD/Source/ChaosVD/ChaosVD.Build.cs b/Engine/Plugins/Experimental/ChaosVD/Source/ChaosVD/ChaosVD.Build.cs
--- a/Engine/Plugins/Experimental/ChaosVD/Source/ChaosVD/ChaosVD.Build.cs
+++ b/Engine/Plugins/Experimental/ChaosVD/Source/ChaosVD/ChaosVD.Build.cs
@@ -12,6 +12,7 @@ public class ChaosVD : ModuleRules
 				"Core",
 				"CoreUObject",
 				"Engine",
+				"RenderCore",
 				"ChaosVDRuntime", 
 				"TraceServices",
 				"ChaosVDData",
diff --git a/Engine/Plugins/Media/MediaFrameworkUtilities/Source/MediaFrameworkUtilities/MediaFrameworkUtilities.Build.cs b/Engine/Plugins/Media/MediaFrameworkUtilities/Source/MediaFrameworkUtilities/MediaFrameworkUtilities.Build.cs
--- a/Engine/Plugins/Media/MediaFrameworkUtilities/Source/MediaFrameworkUtilities/MediaFrameworkUtilities.Build.cs
+++ b/Engine/Plugins/Media/MediaFrameworkUtilities/Source/MediaFrameworkUtilities/MediaFrameworkUtilities.Build.cs
@@ -21,6 +21,7 @@ namespace UnrealBuildTool.Rules
 					"Core",
 					"CoreUObject",
 					"Engine",
+					"RenderCore",
 					"Media",
 					"MediaAssets",
 					"MediaIOCore",
diff --git a/Engine/Plugins/VirtualProduction/CameraCalibration/Source/CameraCalibrationEditor/CameraCalibrationEditor.Build.cs b/Engine/Plugins/VirtualProduction/CameraCalibration/Source/CameraCalibrationEditor/CameraCalibrationEditor.Build.cs
--- a/Engine/Plugins/VirtualProduction/CameraCalibration/Source/CameraCalibrationEditor/CameraCalibrationEditor.Build.cs
+++ b/Engine/Plugins/VirtualProduction/CameraCalibration/Source/CameraCalibrationEditor/CameraCalibrationEditor.Build.cs
@@ -25,6 +25,7 @@ namespace UnrealBuildTool.Rules
 					"EditorStyle",
 					"EditorWidgets",
 					"Engine",
+					"RenderCore",
 					"ImageCore",
 					"InputCore",
 					"Json",
