CameraCalibration: Fix nodal offset aruco detection to use proper corner detection subpixel refinement

diff --git a/Engine/Plugins/VirtualProduction/CameraCalibration/Source/CameraCalibrationEditor/Private/Calibrators/CameraNodalOffsetAlgoAruco.cpp b/Engine/Plugins/VirtualProduction/CameraCalibration/Source/CameraCalibrationEditor/Private/Calibrators/CameraNodalOffsetAlgoAruco.cpp
--- a/Engine/Plugins/VirtualProduction/CameraCalibration/Source/CameraCalibrationEditor/Private/Calibrators/CameraNodalOffsetAlgoAruco.cpp
+++ b/Engine/Plugins/VirtualProduction/CameraCalibration/Source/CameraCalibrationEditor/Private/Calibrators/CameraNodalOffsetAlgoAruco.cpp
@@ -227,6 +227,7 @@ bool UCameraNodalOffsetAlgoAruco::PopulatePoints(FText& OutErrorMessage)
 	std::vector<int> MarkerIds;
 	std::vector<std::vector<cv::Point2f>> MarkerCorners;
 	cv::Ptr<cv::aruco::DetectorParameters> DetectorParameters = cv::aruco::DetectorParameters::create();
+	DetectorParameters->cornerRefinementMethod = cv::aruco::CORNER_REFINE_SUBPIX;
 
 	cv::aruco::detectMarkers(CvGray, DictionaryInfo.Dict, MarkerCorners, MarkerIds, DetectorParameters);
 
