CameraCalibration: Fix distortion calibration solver to properly output final camera poses

diff --git a/Engine/Plugins/Runtime/OpenCV/Source/OpenCVHelper/Private/OpenCVHelper.cpp b/Engine/Plugins/Runtime/OpenCV/Source/OpenCVHelper/Private/OpenCVHelper.cpp
--- a/Engine/Plugins/Runtime/OpenCV/Source/OpenCVHelper/Private/OpenCVHelper.cpp
+++ b/Engine/Plugins/Runtime/OpenCV/Source/OpenCVHelper/Private/OpenCVHelper.cpp
@@ -313,4 +313,65 @@ double FOpenCVHelper::ComputeReprojectionError(const FTransform& CameraPose, con
 	return ReprojectionError;
 }
 
+void FOpenCVHelper::ConvertTransformToVectors(const FTransform& InTransform, cv::Mat& OutRotation, cv::Mat& OutTranslation)
+{
+	const FMatrix Matrix = InTransform.ToMatrixNoScale();
+
+	const cv::Mat Tcam = (cv::Mat_<double>(3, 1) << Matrix.M[3][0], Matrix.M[3][1], Matrix.M[3][2]);
+
+	cv::Mat Rcam = cv::Mat::zeros(3, 3, cv::DataType<double>::type);
+	for (int32 Column = 0; Column < 3; ++Column)
+	{
+		FVector ColVec = Matrix.GetColumn(Column);
+		Rcam.at<double>(Column, 0) = ColVec.X;
+		Rcam.at<double>(Column, 1) = ColVec.Y;
+		Rcam.at<double>(Column, 2) = ColVec.Z;
+	}
+
+	const cv::Mat Robj = Rcam.t();
+
+	cv::Rodrigues(Robj, OutRotation);
+
+	OutTranslation = -Rcam.inv() * Tcam;
+}
+
+void FOpenCVHelper::ConvertVectorsToTransform(const cv::Mat& InRotation, const cv::Mat& InTranslation, FTransform& OutTransform)
+{
+	// Convert to camera pose
+
+	// [R|t]' = [R'|-R'*t]
+
+	// Convert from Rodrigues to rotation matrix
+	cv::Mat Robj;
+	cv::Rodrigues(InRotation, Robj); // Robj is 3x3
+
+	// Calculate camera translation
+	cv::Mat Tcam = -Robj.t() * InTranslation;
+
+	// Invert/transpose to get camera orientation
+	cv::Mat Rcam = Robj.t();
+
+	// Convert back to UE coordinates
+
+	FMatrix M = FMatrix::Identity;
+
+	// Fill rotation matrix
+	for (int32 Column = 0; Column < 3; ++Column)
+	{
+		M.SetColumn(Column, FVector(
+			Rcam.at<double>(Column, 0),
+			Rcam.at<double>(Column, 1),
+			Rcam.at<double>(Column, 2))
+		);
+	}
+
+	// Fill translation vector
+	M.M[3][0] = Tcam.at<double>(0);
+	M.M[3][1] = Tcam.at<double>(1);
+	M.M[3][2] = Tcam.at<double>(2);
+
+	OutTransform.SetFromMatrix(M);
+	FOpenCVHelper::ConvertOpenCVToUnreal(OutTransform);
+}
+
 #endif // WITH_OPENCV
diff --git a/Engine/Plugins/Runtime/OpenCV/Source/OpenCVHelper/Public/OpenCVHelper.h b/Engine/Plugins/Runtime/OpenCV/Source/OpenCVHelper/Public/OpenCVHelper.h
--- a/Engine/Plugins/Runtime/OpenCV/Source/OpenCVHelper/Public/OpenCVHelper.h
+++ b/Engine/Plugins/Runtime/OpenCV/Source/OpenCVHelper/Public/OpenCVHelper.h
@@ -93,6 +93,12 @@ public:
 	static UTexture2D* TextureFromCvMat(cv::Mat& Mat, const FString* PackagePath = nullptr, const FName* TextureName = nullptr);
 	static UTexture2D* TextureFromCvMat(cv::Mat& Mat, UTexture2D* InTexture);
 
+	/** Converts an FTransform to a rotation vector (cv::Mat) and a translation vector (cv::Mat) */
+	static void ConvertTransformToVectors(const FTransform& InTransform, cv::Mat& OutRotation, cv::Mat& OutTranslation);
+
+	/** Converts a rotation vector (cv::Mat) and a translation vector (cv::Mat) to an FTransform */
+	static void ConvertVectorsToTransform(const cv::Mat& InRotation, const cv::Mat& InTranslation, FTransform& OutTransform);
+
 	static double ComputeReprojectionError(const FTransform& CameraPose, const cv::Mat& CameraIntrinsicMatrix, const std::vector<cv::Point3f>& Points3d, const std::vector<cv::Point2f>& Points2d);
 #endif	// WITH_OPENCV
 };
diff --git a/Engine/Plugins/VirtualProduction/CameraCalibration/Source/CameraCalibrationEditor/Private/Calibrators/CameraCalibrationSolver.cpp b/Engine/Plugins/VirtualProduction/CameraCalibration/Source/CameraCalibrationEditor/Private/Calibrators/CameraCalibrationSolver.cpp
--- a/Engine/Plugins/VirtualProduction/CameraCalibration/Source/CameraCalibrationEditor/Private/Calibrators/CameraCalibrationSolver.cpp
+++ b/Engine/Plugins/VirtualProduction/CameraCalibration/Source/CameraCalibrationEditor/Private/Calibrators/CameraCalibrationSolver.cpp
@@ -5,6 +5,7 @@
 #include "HAL/IConsoleManager.h"
 #include "Models/AnamorphicLensModel.h"
 #include "Models/SphericalLensModel.h"
+#include "OpenCVHelper.h"
 
 #if WITH_OPENCV
 
@@ -28,8 +29,9 @@ double FCameraCalibrationSolver::CalibrateCamera(
 	FVector2f& InOutFocalLength,
 	FVector2f& InOutImageCenter,
 	TArray<float>& OutDistCoeffs,
+	TArray<FTransform>& InOutCameraPoses,
 	double PixelAspect,
-	int SolverFlags)
+	ECalibrationFlags SolverFlags)
 {
 	const int NumImages = InObjectPoints.Num();
 
@@ -122,6 +124,20 @@ double FCameraCalibrationSolver::CalibrateCamera(
 			Samples2d.push_back(Points2d);
 		}
 
+		int LegacyFlags = 0;
+		if (EnumHasAnyFlags(SolverFlags, ECalibrationFlags::UseIntrinsicGuess))
+		{
+			LegacyFlags |= cv::CALIB_USE_INTRINSIC_GUESS;
+		}
+		if (EnumHasAnyFlags(SolverFlags, ECalibrationFlags::FixFocalLength))
+		{
+			LegacyFlags |= cv::CALIB_FIX_FOCAL_LENGTH;
+		}
+		if (EnumHasAnyFlags(SolverFlags, ECalibrationFlags::FixPrincipalPoint))
+		{
+			LegacyFlags |= cv::CALIB_FIX_PRINCIPAL_POINT;
+		}
+
 		RMSE = cv::calibrateCamera(
 			Samples3d,
 			Samples2d,
@@ -133,7 +149,7 @@ double FCameraCalibrationSolver::CalibrateCamera(
 			cv::noArray(),
 			cv::noArray(),
 			cv::noArray(),
-			SolverFlags
+			LegacyFlags
 		);
 
 		// Set the output intrinsics and distortion parameters to the final values calculated by the solver
@@ -153,7 +169,7 @@ double FCameraCalibrationSolver::CalibrateCamera(
 	}
 
 	// If the flag to use a starting guess for the camera intrinsics is not set, calculate some initial values for the intrinsic parameters
-	if (!(SolverFlags & cv::CALIB_USE_INTRINSIC_GUESS))
+	if (!(EnumHasAnyFlags(SolverFlags, ECalibrationFlags::UseIntrinsicGuess)))
 	{
 		InitCameraIntrinsics(ObjectPointsMat, ImagePointsMat, NumPointsMat, CvImageSize, CameraMatrix);
 	}
@@ -179,16 +195,30 @@ double FCameraCalibrationSolver::CalibrateCamera(
 
 	// Instruct the solver to ignore some parameters when running its solve
 	uchar* Mask = Solver.Mask.ptr<uchar>();
-	if (SolverFlags & cv::CALIB_FIX_FOCAL_LENGTH)
+	if (EnumHasAnyFlags(SolverFlags, ECalibrationFlags::FixFocalLength))
 	{
 		Mask[0] = 0;
 		Mask[1] = 0;
 	}
-	if (SolverFlags & cv::CALIB_FIX_PRINCIPAL_POINT)
+	if (EnumHasAnyFlags(SolverFlags, ECalibrationFlags::FixPrincipalPoint))
 	{
 		Mask[2] = 0;
 		Mask[3] = 0;
 	}
+	if (EnumHasAnyFlags(SolverFlags, ECalibrationFlags::FixExtrinsics))
+	{
+		for (int ParamIndex = NumIntrinsics; ParamIndex < NumParamsToSolve; ++ParamIndex)
+		{
+			Mask[ParamIndex] = 0;
+		}
+	}
+	if (EnumHasAnyFlags(SolverFlags, ECalibrationFlags::FixZeroDistortion))
+	{
+		for (int ParamIndex = 0; ParamIndex < NumDistortionCoefficients; ++ParamIndex)
+		{
+			Mask[ParamIndex + 4] = 0;
+		}
+	}
 
 	if (LensModel == UAnamorphicLensModel::StaticClass())
 	{
@@ -212,7 +242,14 @@ double FCameraCalibrationSolver::CalibrateCamera(
 		cv::Mat Rotation = Solver.Params.rowRange(ExtrinsicOffset, ExtrinsicOffset + 3);
 		cv::Mat Translation = Solver.Params.rowRange(ExtrinsicOffset + 3, ExtrinsicOffset + 6);
 
-		InitCameraExtrinsics(LensModel, ObjectPointsInImage, ImagePointsInImage, CameraMatrix, DistCoeffs, CvImageSize, Rotation, Translation, SolverFlags);
+		if (!(EnumHasAnyFlags(SolverFlags, ECalibrationFlags::UseExtrinsicGuess)))
+		{
+			InitCameraExtrinsics(LensModel, ObjectPointsInImage, ImagePointsInImage, CameraMatrix, DistCoeffs, CvImageSize, Rotation, Translation, SolverFlags);
+		}
+		else
+		{
+			FOpenCVHelper::ConvertTransformToVectors(InOutCameraPoses[ImageIndex], Rotation, Translation);
+		}
 	}
 
 	double ReprojectionError = 0.0;
@@ -324,6 +361,15 @@ double FCameraCalibrationSolver::CalibrateCamera(
 		OutDistCoeffs.Add(DistCoeffs.at<double>(3));
 	}
 
+	for (int ImageIndex = 0; ImageIndex < NumImages; ImageIndex++)
+	{
+		const int ExtrinsicOffset = NumIntrinsics + (ImageIndex * NumExtrinsics);
+		cv::Mat Rotation = Solver.Params.rowRange(ExtrinsicOffset, ExtrinsicOffset + 3);
+		cv::Mat Translation = Solver.Params.rowRange(ExtrinsicOffset + 3, ExtrinsicOffset + 6);
+
+		FOpenCVHelper::ConvertVectorsToTransform(Rotation, Translation, InOutCameraPoses[ImageIndex]);
+	}
+
 	return RMSE;
 }
 
@@ -436,7 +482,7 @@ void FCameraCalibrationSolver::InitCameraExtrinsics(
 	const cv::Size ImageSize,
 	cv::Mat& Rotation,
 	cv::Mat& Translation,
-	const int SolverFlags)
+	const ECalibrationFlags SolverFlags)
 {
 	const int NumPoints = ObjectPoints.total();
 
@@ -596,7 +642,7 @@ void FCameraCalibrationSolver::ProjectPoints(
 	cv::Mat& ProjectedPoints)
 {
 	cv::Mat Jacobian;
-	const int SolverFlags = 0;
+	const ECalibrationFlags SolverFlags = ECalibrationFlags::None;
 	ProjectPoints(LensModel, ObjectPoints, Rotation, Translation, CameraMatrix, DistCoeffs, ImageSize, ProjectedPoints, Jacobian, SolverFlags);
 }
 
@@ -610,7 +656,7 @@ void FCameraCalibrationSolver::ProjectPoints(
 	const cv::Size ImageSize,
 	cv::Mat& ProjectedPoints,
 	cv::Mat& Jacobian,
-	int SolverFlags)
+	ECalibrationFlags SolverFlags)
 {
 	// Project the 3D points to 2D using the transformation matrix, camera matrix, and distortion equation
 	if (LensModel == USphericalLensModel::StaticClass())
@@ -649,7 +695,7 @@ void FCameraCalibrationSolver::ProjectPointsAnamorphic(
 	const cv::Size ImageSize,
 	cv::Mat& ProjectedPoints,
 	cv::Mat& Jacobian,
-	int SolverFlags)
+	ECalibrationFlags SolverFlags)
 {
 	const bool bComputeJacobian = !Jacobian.empty();
 
@@ -971,7 +1017,7 @@ void FCameraCalibrationSolver::ProjectPointsSpherical(
 	const cv::Mat& DistCoeffs,
 	cv::Mat& ProjectedPoints,
 	cv::Mat& Jacobian,
-	int SolverFlags)
+	ECalibrationFlags SolverFlags)
 {
 	const bool bComputeJacobian = !Jacobian.empty();
 
@@ -1215,7 +1261,7 @@ void FCameraCalibrationSolver::SubdivideJacobian(
 	cv::Mat& JacFocalLength,
 	cv::Mat& JacImageCenter,
 	cv::Mat& JacDistortion,
-	int SolverFlags)
+	ECalibrationFlags SolverFlags)
 {
 	const bool bComputeJacobian = !Jacobian.empty();
 
@@ -1223,18 +1269,21 @@ void FCameraCalibrationSolver::SubdivideJacobian(
 	{
 		if (Jacobian.cols >= 6)
 		{
-			JacRotation = cv::Mat(Jacobian.colRange(0, 3));
-			JacTranslation = cv::Mat(Jacobian.colRange(3, 6));
+			if (!(EnumHasAnyFlags(SolverFlags, ECalibrationFlags::FixExtrinsics)))
+			{
+				JacRotation = cv::Mat(Jacobian.colRange(0, 3));
+				JacTranslation = cv::Mat(Jacobian.colRange(3, 6));
+			}
 		}
 
 		if (Jacobian.cols >= 10)
 		{
-			if (!(SolverFlags & cv::CALIB_FIX_FOCAL_LENGTH))
+			if (!(EnumHasAnyFlags(SolverFlags, ECalibrationFlags::FixFocalLength)))
 			{
 				JacFocalLength = cv::Mat(Jacobian.colRange(6, 8));
 			}
 
-			if (!(SolverFlags & cv::CALIB_FIX_PRINCIPAL_POINT))
+			if (!(EnumHasAnyFlags(SolverFlags, ECalibrationFlags::FixPrincipalPoint)))
 			{
 				JacImageCenter = cv::Mat(Jacobian.colRange(8, 10));
 			}
@@ -1399,6 +1448,11 @@ void FLevMarqSolver::SubMatrix(const cv::Mat& Src, cv::Mat& Dst, const cv::Mat&
 
 void FLevMarqSolver::Step()
 {
+	if (cv::countNonZero(Mask) == 0)
+	{
+		return;
+	}
+
 	cv::Mat SubJtJ;
 	SubMatrix(JacTJac, SubJtJ, Mask, Mask);
 
diff --git a/Engine/Plugins/VirtualProduction/CameraCalibration/Source/CameraCalibrationEditor/Private/Calibrators/CameraCalibrationSolver.h b/Engine/Plugins/VirtualProduction/CameraCalibration/Source/CameraCalibrationEditor/Private/Calibrators/CameraCalibrationSolver.h
--- a/Engine/Plugins/VirtualProduction/CameraCalibration/Source/CameraCalibrationEditor/Private/Calibrators/CameraCalibrationSolver.h
+++ b/Engine/Plugins/VirtualProduction/CameraCalibration/Source/CameraCalibrationEditor/Private/Calibrators/CameraCalibrationSolver.h
@@ -13,6 +13,20 @@
 
 #endif	// WITH_OPENCV
 
+/**
+  * Flags used to modify the execution of the calibration solver
+  */
+enum class ECalibrationFlags : uint32
+{
+	None = 0,
+	UseIntrinsicGuess = 1 << 0,  /** The solver will skip intrinsics initialization an use the input values of Focal Length and Image Center as the starting values for the optimization */
+	UseExtrinsicGuess = 1 << 1,  /** The solver will skip extrinsics initialization an use the input Camera Poses as the starting value for the optimization */
+	FixFocalLength = 1 << 2,     /** The solver will not optimize the focal length */
+	FixPrincipalPoint = 1 << 3,  /** The solver will not optimize the principal point */
+	FixExtrinsics = 1 << 4,      /** The solver will not optimize the camera extrinsics */
+	FixZeroDistortion = 1 << 5   /** The solver will fix all distortion values at 0 */
+};
+
 /** 
   * Lens Distortion Solver class, supporting anamorphic and spherical models 
   * The implementation is largely based on the implementation of calibrateCamera from OpenCV: [REDACTED_PROFILE_URL]/opencv
@@ -33,8 +47,9 @@ public:
 		FVector2f& InOutFocalLength,
 		FVector2f& InOutImageCenter,
 		TArray<float>& OutDistCoeffs,
+		TArray<FTransform>& InOutCameraPoses,
 		double PixelAspect = 1.0,
-		int SolverFlags = 0);
+		ECalibrationFlags SolverFlags = ECalibrationFlags::None);
 
 private:
 	/** Initialize the camera matrix of intrinsic parameters using linear algebra techniques */
@@ -55,7 +70,7 @@ private:
 		const cv::Size ImageSize,
 		cv::Mat& Rotation,
 		cv::Mat& Translation,
-		const int SolverFlags);
+		const ECalibrationFlags SolverFlags);
 
 	/** Project the input object points to 2D using the input camera intrinsics, extrinsics, and distortion parameters */
 	static void ProjectPoints(
@@ -79,7 +94,7 @@ private:
 		const cv::Size ImageSize,
 		cv::Mat& ProjectedPoints,
 		cv::Mat& Jacobian,
-		int SolverFlags);
+		ECalibrationFlags SolverFlags);
 
 	/** Project the input object points to 2D using the input camera intrinsics, extrinsics, and anamorphic distortion parameters */
 	static void ProjectPointsAnamorphic(
@@ -91,7 +106,7 @@ private:
 		const cv::Size ImageSize,
 		cv::Mat& ProjectedPoints,
 		cv::Mat& Jacobian,
-		int SolverFlags);
+		ECalibrationFlags SolverFlags);
 
 	/** Project the input object points to 2D using the input camera intrinsics, extrinsics, and spherical distortion parameters */
 	static void ProjectPointsSpherical(
@@ -102,7 +117,7 @@ private:
 		const cv::Mat& DistCoeffs,
 		cv::Mat& ProjectedPoints,
 		cv::Mat& Jacobian,
-		int SolverFlags);
+		ECalibrationFlags SolverFlags);
 
 	/** Copy the input 3D and 2D points to OpenCV matrices for ease of use with the solver */
 	static void GatherPoints(
@@ -119,7 +134,7 @@ private:
 		cv::Mat& JacFocalLength,
 		cv::Mat& JacImageCenter,
 		cv::Mat& JacDistortion,
-		int SolverFlags);
+		ECalibrationFlags SolverFlags);
 
 #endif	// WITH_OPENCV
 };
diff --git a/Engine/Plugins/VirtualProduction/CameraCalibration/Source/CameraCalibrationEditor/Private/Calibrators/CameraLensDistortionAlgoCheckerboard.cpp b/Engine/Plugins/VirtualProduction/CameraCalibration/Source/CameraCalibrationEditor/Private/Calibrators/CameraLensDistortionAlgoCheckerboard.cpp
--- a/Engine/Plugins/VirtualProduction/CameraCalibration/Source/CameraCalibrationEditor/Private/Calibrators/CameraLensDistortionAlgoCheckerboard.cpp
+++ b/Engine/Plugins/VirtualProduction/CameraCalibration/Source/CameraCalibrationEditor/Private/Calibrators/CameraLensDistortionAlgoCheckerboard.cpp
@@ -6,6 +6,7 @@
 #include "AssetEditor/LensDistortionTool.h"
 #include "AssetEditor/SSimulcamViewport.h"
 #include "AssetRegistry/AssetData.h"
+#include "CalibrationPointComponent.h"
 #include "Camera/CameraActor.h"
 #include "CameraCalibrationCheckerboard.h"
 #include "CameraCalibrationEditorLog.h"
@@ -16,6 +17,7 @@
 #include "EngineUtils.h"
 #include "ImageUtils.h"
 #include "JsonObjectConverter.h"
+#include "LensComponent.h"
 #include "Misc/MessageDialog.h"
 #include "Models/AnamorphicLensModel.h"
 #include "PropertyCustomizationHelpers.h"
@@ -31,10 +33,9 @@
 #define LOCTEXT_NAMESPACE "CameraLensDistortionAlgoCheckerboard"
 
 #if WITH_EDITOR
-static TAutoConsoleVariable<float> CVarFixFocalLengthValue(TEXT("CameraCalibration.FocalLengthValue"), 0.0f, TEXT("Fixed focal length value in mm"));
-static TAutoConsoleVariable<float> CVarSensorWidthValue(TEXT("CameraCalibration.SensorWidthValue"), 0.0f, TEXT("Fixed focal length value in mm"));
-static TAutoConsoleVariable<float> CVarImageWidthValue(TEXT("CameraCalibration.ImageWidthValue"), 0, TEXT("Fixed focal length value in pixels"));
-static TAutoConsoleVariable<float> CVarImageHeightValue(TEXT("CameraCalibration.ImageHeightValue"), 0, TEXT("Fixed focal length value in pixels"));
+static TAutoConsoleVariable<bool> CVarUseIntrinsicsGuess(TEXT("LensDistortionCheckerboard.UseIntrinsicsGuess"), true, TEXT("If true, the solver initializes the camera intrinsics to a user-provided estimate. Otherwise, the solver will compute the initial values."));
+static TAutoConsoleVariable<bool> CVarFixExtrinsics(TEXT("LensDistortionCheckerboard.FixExtrinsics"), false, TEXT("If true, the solver will fix the camera extrinsics to the user-provided camera poses"));
+static TAutoConsoleVariable<bool> CVarFixZeroDistortion(TEXT("LensDistortionCheckerboard.FixZeroDistortion"), false, TEXT("If true, the solver will fix all distortion values to always be 0"));
 #endif
 
 const int UCameraLensDistortionAlgoCheckerboard::DATASET_VERSION = 1;
@@ -182,9 +183,33 @@ void UCameraLensDistortionAlgoCheckerboard::Tick(float DeltaTime)
 			{
 				LastCameraData.InputFocus = LensFileEvalInputs.Focus;
 				LastCameraData.InputZoom = LensFileEvalInputs.Zoom;
-				LastCameraData.bIsValid = true;
 			}
 
+			const ACameraActor* Camera = StepsController->GetCamera();
+
+			if (!Camera)
+			{
+				break;
+			}
+
+			const UCameraComponent* CameraComponent = Camera->GetCameraComponent();
+
+			if (!CameraComponent)
+			{
+				break;
+			}
+
+			LastCameraData.Pose = CameraComponent->GetComponentToWorld();
+
+			const ULensComponent* LensComponent = StepsController->FindLensComponent();
+			if (LensComponent)
+			{
+				LastCameraData.bWasNodalOffsetApplied = LensComponent->WasNodalOffsetAppliedThisTick();
+			}
+
+			LastCameraData.bIsValid = true;
+
+
 		} while (0);
 	}
 }
@@ -298,15 +323,6 @@ bool UCameraLensDistortionAlgoCheckerboard::AddCalibrationRow(FText& OutErrorMes
 	Row->NumCornerCols = Calibrator->NumCornerCols;
 	Row->SquareSideInCm = Calibrator->SquareSideLength;
 
-	// Fill out checkerboard 3d points
-	for (int32 RowIdx = 0; RowIdx < Row->NumCornerRows; ++RowIdx)
-	{
-		for (int32 ColIdx = 0; ColIdx < Row->NumCornerCols; ++ColIdx)
-		{
-			Row->Points3d.Add(Row->SquareSideInCm * FVector(ColIdx, RowIdx, 0));
-		}
-	}
-
 	// Identify checkerboard
 	{
 		cv::Size CheckerboardSize(Row->NumCornerCols, Row->NumCornerRows);
@@ -336,6 +352,12 @@ bool UCameraLensDistortionAlgoCheckerboard::AddCalibrationRow(FText& OutErrorMes
 		cv::TermCriteria Criteria(cv::TermCriteria::Type::EPS | cv::TermCriteria::Type::COUNT, 30, 0.001);
 		cv::cornerSubPix(CvGray, Corners, cv::Size(11, 11), cv::Size(-1, -1), Criteria);
 
+		// Force first one to be top-left
+		if (Corners.front().y > Corners.back().y)
+		{
+			std::reverse(Corners.begin(), Corners.end());
+		}
+
 		for (const cv::Point2f& Corner : Corners)
 		{
 			Row->Points2d.Add(FVector2D(Corner.x, Corner.y));
@@ -384,6 +406,33 @@ bool UCameraLensDistortionAlgoCheckerboard::AddCalibrationRow(FText& OutErrorMes
 		StepsController->RefreshOverlay();
 	}
 
+	// Fill out checkerboard 3d points
+	const FVector LocationTL = Calibrator->TopLeft->GetComponentLocation();
+	const FVector LocationTR = Calibrator->TopRight->GetComponentLocation();
+	const FVector LocationBL = Calibrator->BottomLeft->GetComponentLocation();
+
+	const FVector RightVector = LocationTR - LocationTL;
+	const FVector DownVector = LocationBL - LocationTL;
+
+	// Fill out checkerboard 3d points
+	const float HorizontalStep = (Row->NumCornerCols > 1) ? (1.0f / (Row->NumCornerCols - 1)) : 0.0f;
+	const float VerticalStep = (Row->NumCornerRows > 1) ? (1.0f / (Row->NumCornerRows - 1)) : 0.0f;
+
+	for (int32 RowIdx = 0; RowIdx < Row->NumCornerRows; ++RowIdx)
+	{
+		for (int32 ColIdx = 0; ColIdx < Row->NumCornerCols; ++ColIdx)
+		{
+			const FVector PointLocation = LocationTL + (RightVector * ColIdx * HorizontalStep) + (DownVector * RowIdx * VerticalStep);
+
+			// Convert from UE coordinates to OpenCV coordinates
+			FTransform Transform = FTransform::Identity;
+			Transform.SetLocation(PointLocation);
+			FOpenCVHelper::ConvertUnrealToOpenCV(Transform);
+
+			Row->Points3d.Add(Transform.GetLocation());
+		}
+	}
+
 	// Create thumbnail and add it to the row
 	{
 		// Resize the frame to thumbnail size
@@ -448,11 +497,39 @@ TSharedRef<SWidget> UCameraLensDistortionAlgoCheckerboard::BuildUI()
 		.MaxHeight(FCameraCalibrationWidgetHelpers::DefaultRowHeight)
 		[FCameraCalibrationWidgetHelpers::BuildLabelWidgetPair(LOCTEXT("ShowDetection", "Show Detection"), BuildShowDetectionWidget())]
 
-	+ SVerticalBox::Slot() // Fix Focal Length
+	+ SVerticalBox::Slot() // Solver Settings Title
+		.Padding(0, 5)
+		.AutoHeight()
+		.VAlign(EVerticalAlignment::VAlign_Center)
+		[
+			SNew(SBox) // Constrain the height
+			.MinDesiredHeight(FCameraCalibrationWidgetHelpers::DefaultRowHeight)
+			.MaxDesiredHeight(FCameraCalibrationWidgetHelpers::DefaultRowHeight)
+			[
+				SNew(SBorder) // Background color of title
+				.BorderImage(FAppStyle::GetBrush("DetailsView.CategoryTop"))
+				.BorderBackgroundColor(FLinearColor::White)
+				.VAlign(EVerticalAlignment::VAlign_Center)
+				[
+					SNew(SOverlay)
+					+ SOverlay::Slot()
+					.Padding(5, 0, 0, 0)
+					[
+						SNew(STextBlock) // Title text
+						.Text(LOCTEXT("SolverSettings", "Solver Settings"))
+						.TransformPolicy(ETextTransformPolicy::ToUpper)
+						.Font(FAppStyle::Get().GetFontStyle(TEXT("PropertyWindow.BoldFont")))
+						.TextStyle(FAppStyle::Get(), "DetailsView.CategoryTextStyle")
+					]
+				]
+			]
+		]
+
+	+ SVerticalBox::Slot() // Focal Length Estimate
 		.VAlign(EVerticalAlignment::VAlign_Top)
 		.AutoHeight()
 		.MaxHeight(FCameraCalibrationWidgetHelpers::DefaultRowHeight)
-		[FCameraCalibrationWidgetHelpers::BuildLabelWidgetPair(LOCTEXT("FixFocalLength", "Fix Focal Length"), BuildFixFocalLengthWidget())]
+		[FCameraCalibrationWidgetHelpers::BuildLabelWidgetPair(LOCTEXT("FocalLengthEstimate", "Focal Length Estimate"), BuildFocalLengthEstimateWidget())]
 
 	+ SVerticalBox::Slot() // Fix Image Center
 		.VAlign(EVerticalAlignment::VAlign_Top)
@@ -460,6 +537,39 @@ TSharedRef<SWidget> UCameraLensDistortionAlgoCheckerboard::BuildUI()
 		.MaxHeight(FCameraCalibrationWidgetHelpers::DefaultRowHeight)
 		[FCameraCalibrationWidgetHelpers::BuildLabelWidgetPair(LOCTEXT("FixImageCenter", "Fix Image Center"), BuildFixImageCenterWidget())]
 
+	+ SVerticalBox::Slot() // Tracked Camera/Calibrator Settings Message
+		.Padding(0, 5)
+		.AutoHeight()
+		.VAlign(EVerticalAlignment::VAlign_Center)
+		[
+			SNew(SBox) // Constrain the height
+			.MinDesiredHeight(FCameraCalibrationWidgetHelpers::DefaultRowHeight)
+			.MaxDesiredHeight(FCameraCalibrationWidgetHelpers::DefaultRowHeight)
+			.HAlign(HAlign_Center)
+			[
+				SNew(SBorder)
+				.VAlign(EVerticalAlignment::VAlign_Center)
+				.Padding(5)
+				[
+					SNew(STextBlock)
+					.Text(LOCTEXT("TrackedCameraOptions", "Only use these settings if both the camera and calibrator are tracked"))
+					.Justification(ETextJustify::Center)
+				]
+			]
+		]
+
+	+ SVerticalBox::Slot() // Use Extrinsics Guess
+		.VAlign(EVerticalAlignment::VAlign_Top)
+		.AutoHeight()
+		.MaxHeight(FCameraCalibrationWidgetHelpers::DefaultRowHeight)
+		[FCameraCalibrationWidgetHelpers::BuildLabelWidgetPair(LOCTEXT("UseExtrinsicsGuess", "Use Tracked Camera Pose"), BuildExtrinsicsGuessWidget())]
+
+	+ SVerticalBox::Slot() // Calibrate Camera Pose
+		.VAlign(EVerticalAlignment::VAlign_Top)
+		.AutoHeight()
+		.MaxHeight(FCameraCalibrationWidgetHelpers::DefaultRowHeight)
+		[FCameraCalibrationWidgetHelpers::BuildLabelWidgetPair(LOCTEXT("CalibrateCameraPose", "Calibrate Nodal Offset"), BuildCalibrateNodalOffsetWidget())]
+
 	+ SVerticalBox::Slot() // Calibration Rows
 		.AutoHeight()
 		.MaxHeight(12 * FCameraCalibrationWidgetHelpers::DefaultRowHeight)
@@ -611,7 +721,7 @@ bool UCameraLensDistortionAlgoCheckerboard::GetLensDistortion(
 		return false;
 	}
 
-	const ULensFile* LensFile = StepsController->GetLensFile();
+	ULensFile* LensFile = StepsController->GetLensFile();
 
 	if (!ensure(LensFile))
 	{
@@ -634,59 +744,65 @@ bool UCameraLensDistortionAlgoCheckerboard::GetLensDistortion(
 	float TestImageWidth = LensFile->CameraFeedInfo.GetDimensions().X;
 	float TestImageHeight = LensFile->CameraFeedInfo.GetDimensions().Y;
 
-	float TestFocalLength = StepsController->GetLensFileEvaluationInputs().Zoom;
 	float TestSensorWidth = StepsController->GetLensFileEvaluationInputs().Filmback.SensorWidth;
 
 	const float PixelAspect = LensFile->LensInfo.SqueezeFactor;
 
-	if (CVarFixFocalLengthValue.GetValueOnGameThread() > 1.0f)
-	{
-		TestFocalLength = CVarFixFocalLengthValue.GetValueOnGameThread();
-	}
-	if (CVarSensorWidthValue.GetValueOnGameThread() > 1.0f)
-	{
-		TestSensorWidth = CVarSensorWidthValue.GetValueOnGameThread();
-	}
-	if (CVarImageWidthValue.GetValueOnGameThread() > 1)
-	{
-		TestImageWidth = CVarImageWidthValue.GetValueOnGameThread();
-	}
-	if (CVarImageHeightValue.GetValueOnGameThread() > 1)
-	{
-		TestImageHeight = CVarImageHeightValue.GetValueOnGameThread();
-	}
-
 	TestSensorWidth = TestSensorWidth * PixelAspect;
 
 	float Fx = TestImageWidth;
 	if (!FMath::IsNearlyZero(TestSensorWidth))
 	{
-		Fx = (TestFocalLength / TestSensorWidth) * TestImageWidth;
+		Fx = (FocalLengthEstimate / TestSensorWidth) * TestImageWidth;
 	}
 
 	FVector2f FocalLength = FVector2f(Fx, Fx);
-	FVector2f ImageCenter = FVector2f(TestImageWidth * 0.5f, TestImageHeight * 0.5f);
+	FVector2f ImageCenter = FVector2f((TestImageWidth - 1) * 0.5f, (TestImageHeight - 1) * 0.5f);
 
 	TArray<TArray<FVector2D>> Samples2d;
 	TArray<TArray<FVector>> Samples3d;
 
+	TArray<FTransform> CameraPoses;
+
 	for (const TSharedPtr<FLensDistortionCheckerboardRowData>& Row : CalibrationRows)
 	{
 		Samples2d.Add(Row->Points2d);
 		Samples3d.Add(Row->Points3d);
+
+		FOpenCVHelper::ConvertUnrealToOpenCV(Row->CameraData.Pose);
+		CameraPoses.Add(Row->CameraData.Pose);
+	}
+
+	ECalibrationFlags SolverFlags = ECalibrationFlags::None;
+
+	if (bUseExtrinsicsGuess)
+	{
+		EnumAddFlags(SolverFlags, ECalibrationFlags::UseExtrinsicGuess);
+	}
+
+	if (CVarUseIntrinsicsGuess.GetValueOnGameThread())
+	{
+		EnumAddFlags(SolverFlags, ECalibrationFlags::UseIntrinsicGuess);
+	}
+
+	if (CVarFixExtrinsics.GetValueOnGameThread())
+	{
+		EnumAddFlags(SolverFlags, ECalibrationFlags::FixExtrinsics);
+	}
+
+	if (CVarFixZeroDistortion.GetValueOnGameThread())
+	{
+		EnumAddFlags(SolverFlags, ECalibrationFlags::FixZeroDistortion);
 	}
 
-	int flags = cv::CALIB_USE_INTRINSIC_GUESS;
 	if (bFixFocalLength)
 	{
-		flags |= cv::CALIB_USE_INTRINSIC_GUESS;
-		flags |= cv::CALIB_FIX_FOCAL_LENGTH;
+		EnumAddFlags(SolverFlags, ECalibrationFlags::FixFocalLength);
 	}
 
 	if (bFixImageCenter)
 	{
-		flags |= cv::CALIB_USE_INTRINSIC_GUESS;
-		flags |= cv::CALIB_FIX_PRINCIPAL_POINT;
+		EnumAddFlags(SolverFlags, ECalibrationFlags::FixPrincipalPoint);
 	}
 
 	TArray<float> DistortionCoefficients;
@@ -699,8 +815,9 @@ bool UCameraLensDistortionAlgoCheckerboard::GetLensDistortion(
 		FocalLength,
 		ImageCenter,
 		DistortionCoefficients,
+		CameraPoses,
 		PixelAspect,
-		flags
+		SolverFlags
 	);
 
 	OutLensModel = LensFile->LensInfo.LensModel;
@@ -718,6 +835,42 @@ bool UCameraLensDistortionAlgoCheckerboard::GetLensDistortion(
 		float(ImageCenter.Y / LastRow->ImageHeight)
 	);
 
+	if (bCalibrateNodalOffset)
+	{
+		// See if the camera already had an offset applied, in which case we need to account for it.
+		FTransform ExistingOffset = FTransform::Identity;
+
+		if (CalibrationRows[0]->CameraData.bWasNodalOffsetApplied)
+		{
+			FNodalPointOffset NodalPointOffset;
+
+			if (LensFile->EvaluateNodalPointOffset(OutFocus, OutZoom, NodalPointOffset))
+			{
+				ExistingOffset.SetTranslation(NodalPointOffset.LocationOffset);
+				ExistingOffset.SetRotation(NodalPointOffset.RotationOffset);
+			}
+		}
+
+		TArray<FNodalPointOffset> NewNodalOffsets;
+		NewNodalOffsets.Reserve(CameraPoses.Num());
+
+		for (int32 RowIndex = 0; RowIndex < CalibrationRows.Num(); ++RowIndex)
+		{
+			FOpenCVHelper::ConvertOpenCVToUnreal(CalibrationRows[RowIndex]->CameraData.Pose);
+
+			const FTransform DesiredCameraTransform = CameraPoses[RowIndex];
+			const FTransform DesiredOffset = DesiredCameraTransform * CalibrationRows[RowIndex]->CameraData.Pose.Inverse() * ExistingOffset;
+
+			FNodalPointOffset NewNodalOffset;
+			NewNodalOffset.LocationOffset = DesiredOffset.GetLocation();
+			NewNodalOffset.RotationOffset = DesiredOffset.GetRotation();
+
+			NewNodalOffsets.Add(NewNodalOffset);
+		}
+
+		LensFile->AddNodalOffsetPoint(OutFocus, OutZoom, NewNodalOffsets[0]);
+	}
+
 	// FZ inputs to LUT
 	OutFocus = LastRow->CameraData.InputFocus;
 	OutZoom = LastRow->CameraData.InputZoom;
@@ -826,29 +979,75 @@ TSharedRef<SWidget> UCameraLensDistortionAlgoCheckerboard::BuildShowDetectionWid
 	});
 }
 
-TSharedRef<SWidget> UCameraLensDistortionAlgoCheckerboard::BuildFixFocalLengthWidget()
+TSharedRef<SWidget> UCameraLensDistortionAlgoCheckerboard::BuildFocalLengthEstimateWidget()
+{
+	return SNew(SHorizontalBox)
+	
+	+ SHorizontalBox::Slot()
+	.FillWidth(0.8)
+	.Padding(0, 0, 10, 0)
+	[
+		SNew(SNumericEntryBox<float>)
+		.Value(MakeAttributeLambda([this]() { return TOptional<float>(FocalLengthEstimate); }))
+		.OnValueChanged(SNumericEntryBox<float>::FOnValueChanged::CreateLambda([this](float NewValue) { FocalLengthEstimate = NewValue; }))
+	]
+
+	+ SHorizontalBox::Slot()
+	.FillWidth(0.2)
+	[
+		SNew(SCheckBox)
+		.Padding(FMargin(5, 0, 15, 0))
+		.IsChecked_Lambda([&]() -> ECheckBoxState
+		{
+			return bFixFocalLength ? ECheckBoxState::Checked : ECheckBoxState::Unchecked;
+		})
+		.OnCheckStateChanged_Lambda([&](ECheckBoxState NewState) -> void
+		{
+			bFixFocalLength = (NewState == ECheckBoxState::Checked);
+		})
+		[
+			SNew(STextBlock)
+			.Text(LOCTEXT("FixText", "Fix?"))
+		]
+	];
+}
+
+TSharedRef<SWidget> UCameraLensDistortionAlgoCheckerboard::BuildFixImageCenterWidget()
 {
 	return SNew(SCheckBox)
 	.IsChecked_Lambda([&]() -> ECheckBoxState
 	{
-		return bFixFocalLength ? ECheckBoxState::Checked : ECheckBoxState::Unchecked;
+		return bFixImageCenter ? ECheckBoxState::Checked : ECheckBoxState::Unchecked;
 	})
 	.OnCheckStateChanged_Lambda([&](ECheckBoxState NewState) -> void
 	{
-		bFixFocalLength = (NewState == ECheckBoxState::Checked);
+		bFixImageCenter = (NewState == ECheckBoxState::Checked);
 	});
 }
 
-TSharedRef<SWidget> UCameraLensDistortionAlgoCheckerboard::BuildFixImageCenterWidget()
+TSharedRef<SWidget> UCameraLensDistortionAlgoCheckerboard::BuildExtrinsicsGuessWidget()
 {
 	return SNew(SCheckBox)
 	.IsChecked_Lambda([&]() -> ECheckBoxState
 	{
-		return bFixImageCenter ? ECheckBoxState::Checked : ECheckBoxState::Unchecked;
+		return bUseExtrinsicsGuess ? ECheckBoxState::Checked : ECheckBoxState::Unchecked;
 	})
 	.OnCheckStateChanged_Lambda([&](ECheckBoxState NewState) -> void
 	{
-		bFixImageCenter = (NewState == ECheckBoxState::Checked);
+		bUseExtrinsicsGuess = (NewState == ECheckBoxState::Checked);
+	});
+}
+
+TSharedRef<SWidget> UCameraLensDistortionAlgoCheckerboard::BuildCalibrateNodalOffsetWidget()
+{
+	return SNew(SCheckBox)
+	.IsChecked_Lambda([&]() -> ECheckBoxState
+	{
+		return bCalibrateNodalOffset ? ECheckBoxState::Checked : ECheckBoxState::Unchecked;
+	})
+	.OnCheckStateChanged_Lambda([&](ECheckBoxState NewState) -> void
+	{
+		bCalibrateNodalOffset = (NewState == ECheckBoxState::Checked);
 	});
 }
 
diff --git a/Engine/Plugins/VirtualProduction/CameraCalibration/Source/CameraCalibrationEditor/Private/Calibrators/CameraLensDistortionAlgoCheckerboard.h b/Engine/Plugins/VirtualProduction/CameraCalibration/Source/CameraCalibrationEditor/Private/Calibrators/CameraLensDistortionAlgoCheckerboard.h
--- a/Engine/Plugins/VirtualProduction/CameraCalibration/Source/CameraCalibrationEditor/Private/Calibrators/CameraLensDistortionAlgoCheckerboard.h
+++ b/Engine/Plugins/VirtualProduction/CameraCalibration/Source/CameraCalibrationEditor/Private/Calibrators/CameraLensDistortionAlgoCheckerboard.h
@@ -62,6 +62,14 @@ struct FLensDistortionCheckerboardCameraData
 	// Input zoom from lens file evaluation data
 	UPROPERTY()
 	float InputZoom = 0.0f;
+
+	// The camera pose
+	UPROPERTY()
+	FTransform Pose = FTransform::Identity;
+
+	// True if the cached camera pose included a nodal offset
+	UPROPERTY()
+	bool bWasNodalOffsetApplied = false;
 };
 
 /** Holds information of the calibration row */
@@ -192,6 +200,15 @@ private:
 	/** If true, the solver will not solve for image center, and will just use the input value */
 	bool bFixImageCenter = false;
 
+	/** Estimated focal length to initialize the distortion solver */
+	float FocalLengthEstimate = 0.0f;
+
+	/** If true, the solver will use the current camera pose to initialize the camera extrinsic parameters for each image */
+	bool bUseExtrinsicsGuess = false;
+
+	/** If true, the final camera poses generated by the solver will be used to compute the nodal offset from the initial camera poses  */
+	bool bCalibrateNodalOffset = false;
+
 	/** Texture into which detected chessboard corners from each calibration row are drawn */
 	TObjectPtr<UTexture2D> CoverageTexture;
 
@@ -211,12 +228,18 @@ private:
 	/** Builds the UI for the user to select if they want a corner detection window to be shown after every capture */
 	TSharedRef<SWidget> BuildShowDetectionWidget();
 
-	/** Builds the UI of the fix focal length width */
-	TSharedRef<SWidget> BuildFixFocalLengthWidget();
+	/** Builds the UI of the focal length estimate widget */
+	TSharedRef<SWidget> BuildFocalLengthEstimateWidget();
 
 	/** Builds the UI of the fix image center widget */
 	TSharedRef<SWidget> BuildFixImageCenterWidget();
 
+	/** Builds the UI of the extrinsics guess widget */
+	TSharedRef<SWidget> BuildExtrinsicsGuessWidget();
+
+	/** Builds the UI of the calibrate nodal offset widget */
+	TSharedRef<SWidget> BuildCalibrateNodalOffsetWidget();
+
 	/** Builds the UI of the calibration points table */
 	TSharedRef<SWidget> BuildCalibrationPointsTable();
 
