CameraCalibration: Fix crash in nodal offset tool when selecting a calibration component but no calibrator is selected.

diff --git a/Engine/Plugins/VirtualProduction/CameraCalibration/Source/CameraCalibrationEditor/Private/Calibrators/CameraNodalOffsetAlgoPoints.cpp b/Engine/Plugins/VirtualProduction/CameraCalibration/Source/CameraCalibrationEditor/Private/Calibrators/CameraNodalOffsetAlgoPoints.cpp
--- a/Engine/Plugins/VirtualProduction/CameraCalibration/Source/CameraCalibrationEditor/Private/Calibrators/CameraNodalOffsetAlgoPoints.cpp
+++ b/Engine/Plugins/VirtualProduction/CameraCalibration/Source/CameraCalibrationEditor/Private/Calibrators/CameraNodalOffsetAlgoPoints.cpp
@@ -1424,9 +1424,12 @@ TSharedRef<SWidget> UCameraNodalOffsetAlgoPoints::BuildCalibrationComponentMenu(
 	// Generate menu
 	FMenuBuilder MenuBuilder(true, nullptr);
 	MenuBuilder.BeginSection("CalibrationComponents", LOCTEXT("CalibrationComponents", "Calibration Point Components"));
+
+	AActor* CalibratorPtr = Calibrator.Get();
+	if (CalibratorPtr)
 	{
 		TArray<UCalibrationPointComponent*, TInlineAllocator<NumInlineAllocations>> CalibrationPointComponents;
-		Calibrator->GetComponents(CalibrationPointComponents);
+		CalibratorPtr->GetComponents(CalibrationPointComponents);
 
 		for (UCalibrationPointComponent* CalibratorComponent : CalibrationPointComponents)
 		{
@@ -1447,6 +1450,21 @@ TSharedRef<SWidget> UCameraNodalOffsetAlgoPoints::BuildCalibrationComponentMenu(
 			}
 		}
 	}
+	else
+	{
+		MenuBuilder.AddMenuEntry(
+			FText::FromName(NAME_None),
+			FText::FromName(NAME_None),
+			FSlateIcon(),
+			FUIAction(
+				FExecuteAction(),
+				FCanExecuteAction(),
+				FIsActionChecked()
+			),
+			NAME_None,
+			EUserInterfaceActionType::None
+		);
+	}
 	MenuBuilder.EndSection();
 
 	return MenuBuilder.MakeWidget();
