CameraCalibration: Fix nodal offset calibration to save to correct focus/zoom

diff --git a/Engine/Plugins/VirtualProduction/CameraCalibration/Source/CameraCalibrationEditor/Private/Calibrators/CameraLensDistortionAlgoCheckerboard.cpp b/Engine/Plugins/VirtualProduction/CameraCalibration/Source/CameraCalibrationEditor/Private/Calibrators/CameraLensDistortionAlgoCheckerboard.cpp
--- a/Engine/Plugins/VirtualProduction/CameraCalibration/Source/CameraCalibrationEditor/Private/Calibrators/CameraLensDistortionAlgoCheckerboard.cpp
+++ b/Engine/Plugins/VirtualProduction/CameraCalibration/Source/CameraCalibrationEditor/Private/Calibrators/CameraLensDistortionAlgoCheckerboard.cpp
@@ -562,13 +562,49 @@ TSharedRef<SWidget> UCameraLensDistortionAlgoCheckerboard::BuildUI()
 		.VAlign(EVerticalAlignment::VAlign_Top)
 		.AutoHeight()
 		.MaxHeight(FCameraCalibrationWidgetHelpers::DefaultRowHeight)
-		[FCameraCalibrationWidgetHelpers::BuildLabelWidgetPair(LOCTEXT("UseExtrinsicsGuess", "Use Tracked Camera Pose"), BuildExtrinsicsGuessWidget())]
+		[
+			SNew(SHorizontalBox)
+
+			+ SHorizontalBox::Slot()
+				.VAlign(VAlign_Center)
+				.Padding(5, 5)
+				.FillWidth(0.35f)
+				[
+					SNew(STextBlock)
+					.Text(LOCTEXT("UseExtrinsicsGuess", "Use Tracked Camera Pose"))
+					.ToolTipText(LOCTEXT("UseExtrinsicsGuessTooltip", "If checked, the solver will use the current transform of the camera as its initial guess for the camera's pose."))
+				]
+
+			+ SHorizontalBox::Slot()
+				.VAlign(VAlign_Center)
+				.Padding(5, 5)
+				.FillWidth(0.65f)
+				[BuildExtrinsicsGuessWidget()]
+		]
 
 	+ SVerticalBox::Slot() // Calibrate Camera Pose
 		.VAlign(EVerticalAlignment::VAlign_Top)
 		.AutoHeight()
 		.MaxHeight(FCameraCalibrationWidgetHelpers::DefaultRowHeight)
-		[FCameraCalibrationWidgetHelpers::BuildLabelWidgetPair(LOCTEXT("CalibrateCameraPose", "Calibrate Nodal Offset"), BuildCalibrateNodalOffsetWidget())]
+		[
+			SNew(SHorizontalBox)
+
+			+ SHorizontalBox::Slot()
+				.VAlign(VAlign_Center)
+				.Padding(5, 5)
+				.FillWidth(0.35f)
+				[
+					SNew(STextBlock)
+					.Text(LOCTEXT("CalibrateCameraPose", "Calibrate Nodal Offset"))
+					.ToolTipText(LOCTEXT("UseExtrinsicsGuessTooltip", "If checked, the final camera pose computed by the solver will be used to save the nodal offset to the Lens File."))
+				]
+
+			+ SHorizontalBox::Slot()
+				.VAlign(VAlign_Center)
+				.Padding(5, 5)
+				.FillWidth(0.65f)
+				[BuildCalibrateNodalOffsetWidget()]
+		]
 
 	+ SVerticalBox::Slot() // Calibration Rows
 		.AutoHeight()
@@ -835,6 +871,10 @@ bool UCameraLensDistortionAlgoCheckerboard::GetLensDistortion(
 		float(ImageCenter.Y / LastRow->ImageHeight)
 	);
 
+	// FZ inputs to LUT
+	OutFocus = LastRow->CameraData.InputFocus;
+	OutZoom = LastRow->CameraData.InputZoom;
+
 	if (bCalibrateNodalOffset)
 	{
 		// See if the camera already had an offset applied, in which case we need to account for it.
@@ -871,10 +911,6 @@ bool UCameraLensDistortionAlgoCheckerboard::GetLensDistortion(
 		LensFile->AddNodalOffsetPoint(OutFocus, OutZoom, NewNodalOffsets[0]);
 	}
 
-	// FZ inputs to LUT
-	OutFocus = LastRow->CameraData.InputFocus;
-	OutZoom = LastRow->CameraData.InputZoom;
-
 	return true;
 #else
 	{
