CameraCalibration: Fix CIS Localization warning for duplicate key

diff --git a/Engine/Plugins/VirtualProduction/CameraCalibration/Source/CameraCalibrationEditor/Private/Calibrators/CameraLensDistortionAlgoCheckerboard.cpp b/Engine/Plugins/VirtualProduction/CameraCalibration/Source/CameraCalibrationEditor/Private/Calibrators/CameraLensDistortionAlgoCheckerboard.cpp
--- a/Engine/Plugins/VirtualProduction/CameraCalibration/Source/CameraCalibrationEditor/Private/Calibrators/CameraLensDistortionAlgoCheckerboard.cpp
+++ b/Engine/Plugins/VirtualProduction/CameraCalibration/Source/CameraCalibrationEditor/Private/Calibrators/CameraLensDistortionAlgoCheckerboard.cpp
@@ -596,7 +596,7 @@ TSharedRef<SWidget> UCameraLensDistortionAlgoCheckerboard::BuildUI()
 				[
 					SNew(STextBlock)
 					.Text(LOCTEXT("CalibrateCameraPose", "Calibrate Nodal Offset"))
-					.ToolTipText(LOCTEXT("UseExtrinsicsGuessTooltip", "If checked, the final camera pose computed by the solver will be used to save the nodal offset to the Lens File."))
+					.ToolTipText(LOCTEXT("CalibrateNodalOffsetTooltip", "If checked, the final camera pose computed by the solver will be used to save the nodal offset to the Lens File."))
 				]
 
 			+ SHorizontalBox::Slot()
