Add cvar to fix bad distortion calibration when not using tracked camera pose

diff --git a/Engine/Plugins/VirtualProduction/CameraCalibration/Source/CameraCalibrationEditor/Private/Calibrators/CameraLensDistortionAlgoCheckerboard.cpp b/Engine/Plugins/VirtualProduction/CameraCalibration/Source/CameraCalibrationEditor/Private/Calibrators/CameraLensDistortionAlgoCheckerboard.cpp
--- a/Engine/Plugins/VirtualProduction/CameraCalibration/Source/CameraCalibrationEditor/Private/Calibrators/CameraLensDistortionAlgoCheckerboard.cpp
+++ b/Engine/Plugins/VirtualProduction/CameraCalibration/Source/CameraCalibrationEditor/Private/Calibrators/CameraLensDistortionAlgoCheckerboard.cpp
@@ -36,6 +36,7 @@
 static TAutoConsoleVariable<bool> CVarUseIntrinsicsGuess(TEXT("LensDistortionCheckerboard.UseIntrinsicsGuess"), true, TEXT("If true, the solver initializes the camera intrinsics to a user-provided estimate. Otherwise, the solver will compute the initial values."));
 static TAutoConsoleVariable<bool> CVarFixExtrinsics(TEXT("LensDistortionCheckerboard.FixExtrinsics"), false, TEXT("If true, the solver will fix the camera extrinsics to the user-provided camera poses"));
 static TAutoConsoleVariable<bool> CVarFixZeroDistortion(TEXT("LensDistortionCheckerboard.FixZeroDistortion"), false, TEXT("If true, the solver will fix all distortion values to always be 0"));
+static TAutoConsoleVariable<bool> CVarUseExtrinsicsGuess(TEXT("LensDistortionCheckerboard.UseExtrinsicsGuess"), false, TEXT("If true, the actual checkerboard and camera poses will be used when running the solver"));
 #endif
 
 const int UCameraLensDistortionAlgoCheckerboard::DATASET_VERSION = 1;
@@ -406,30 +407,44 @@ bool UCameraLensDistortionAlgoCheckerboard::AddCalibrationRow(FText& OutErrorMes
 		StepsController->RefreshOverlay();
 	}
 
-	// Fill out checkerboard 3d points
-	const FVector LocationTL = Calibrator->TopLeft->GetComponentLocation();
-	const FVector LocationTR = Calibrator->TopRight->GetComponentLocation();
-	const FVector LocationBL = Calibrator->BottomLeft->GetComponentLocation();
+	if (CVarUseExtrinsicsGuess.GetValueOnGameThread())
+	{
+		// Fill out checkerboard 3d points
+		const FVector LocationTL = Calibrator->TopLeft->GetComponentLocation();
+		const FVector LocationTR = Calibrator->TopRight->GetComponentLocation();
+		const FVector LocationBL = Calibrator->BottomLeft->GetComponentLocation();
 
-	const FVector RightVector = LocationTR - LocationTL;
-	const FVector DownVector = LocationBL - LocationTL;
+		const FVector RightVector = LocationTR - LocationTL;
+		const FVector DownVector = LocationBL - LocationTL;
 
-	// Fill out checkerboard 3d points
-	const float HorizontalStep = (Row->NumCornerCols > 1) ? (1.0f / (Row->NumCornerCols - 1)) : 0.0f;
-	const float VerticalStep = (Row->NumCornerRows > 1) ? (1.0f / (Row->NumCornerRows - 1)) : 0.0f;
+		// Fill out checkerboard 3d points
+		const float HorizontalStep = (Row->NumCornerCols > 1) ? (1.0f / (Row->NumCornerCols - 1)) : 0.0f;
+		const float VerticalStep = (Row->NumCornerRows > 1) ? (1.0f / (Row->NumCornerRows - 1)) : 0.0f;
 
-	for (int32 RowIdx = 0; RowIdx < Row->NumCornerRows; ++RowIdx)
-	{
-		for (int32 ColIdx = 0; ColIdx < Row->NumCornerCols; ++ColIdx)
+		for (int32 RowIdx = 0; RowIdx < Row->NumCornerRows; ++RowIdx)
 		{
-			const FVector PointLocation = LocationTL + (RightVector * ColIdx * HorizontalStep) + (DownVector * RowIdx * VerticalStep);
+			for (int32 ColIdx = 0; ColIdx < Row->NumCornerCols; ++ColIdx)
+			{
+				const FVector PointLocation = LocationTL + (RightVector * ColIdx * HorizontalStep) + (DownVector * RowIdx * VerticalStep);
 
-			// Convert from UE coordinates to OpenCV coordinates
-			FTransform Transform = FTransform::Identity;
-			Transform.SetLocation(PointLocation);
-			FOpenCVHelper::ConvertUnrealToOpenCV(Transform);
+				// Convert from UE coordinates to OpenCV coordinates
+				FTransform Transform = FTransform::Identity;
+				Transform.SetLocation(PointLocation);
+				FOpenCVHelper::ConvertUnrealToOpenCV(Transform);
 
-			Row->Points3d.Add(Transform.GetLocation());
+				Row->Points3d.Add(Transform.GetLocation());
+			}
+		}
+	}
+	else
+	{
+		// Fill out checkerboard 3d points
+		for (int32 RowIdx = 0; RowIdx < Row->NumCornerRows; ++RowIdx)
+		{
+			for (int32 ColIdx = 0; ColIdx < Row->NumCornerCols; ++ColIdx)
+			{
+				Row->Points3d.Add(Row->SquareSideInCm * FVector(ColIdx, RowIdx, 0));
+			}
 		}
 	}
 
@@ -537,75 +552,6 @@ TSharedRef<SWidget> UCameraLensDistortionAlgoCheckerboard::BuildUI()
 		.MaxHeight(FCameraCalibrationWidgetHelpers::DefaultRowHeight)
 		[FCameraCalibrationWidgetHelpers::BuildLabelWidgetPair(LOCTEXT("FixImageCenter", "Fix Image Center"), BuildFixImageCenterWidget())]
 
-	+ SVerticalBox::Slot() // Tracked Camera/Calibrator Settings Message
-		.Padding(0, 5)
-		.AutoHeight()
-		.VAlign(EVerticalAlignment::VAlign_Center)
-		[
-			SNew(SBox) // Constrain the height
-			.MinDesiredHeight(FCameraCalibrationWidgetHelpers::DefaultRowHeight)
-			.MaxDesiredHeight(FCameraCalibrationWidgetHelpers::DefaultRowHeight)
-			.HAlign(HAlign_Center)
-			[
-				SNew(SBorder)
-				.VAlign(EVerticalAlignment::VAlign_Center)
-				.Padding(5)
-				[
-					SNew(STextBlock)
-					.Text(LOCTEXT("TrackedCameraOptions", "Only use these settings if both the camera and calibrator are tracked"))
-					.Justification(ETextJustify::Center)
-				]
-			]
-		]
-
-	+ SVerticalBox::Slot() // Use Extrinsics Guess
-		.VAlign(EVerticalAlignment::VAlign_Top)
-		.AutoHeight()
-		.MaxHeight(FCameraCalibrationWidgetHelpers::DefaultRowHeight)
-		[
-			SNew(SHorizontalBox)
-
-			+ SHorizontalBox::Slot()
-				.VAlign(VAlign_Center)
-				.Padding(5, 5)
-				.FillWidth(0.35f)
-				[
-					SNew(STextBlock)
-					.Text(LOCTEXT("UseExtrinsicsGuess", "Use Tracked Camera Pose"))
-					.ToolTipText(LOCTEXT("UseExtrinsicsGuessTooltip", "If checked, the solver will use the current transform of the camera as its initial guess for the camera's pose."))
-				]
-
-			+ SHorizontalBox::Slot()
-				.VAlign(VAlign_Center)
-				.Padding(5, 5)
-				.FillWidth(0.65f)
-				[BuildExtrinsicsGuessWidget()]
-		]
-
-	+ SVerticalBox::Slot() // Calibrate Camera Pose
-		.VAlign(EVerticalAlignment::VAlign_Top)
-		.AutoHeight()
-		.MaxHeight(FCameraCalibrationWidgetHelpers::DefaultRowHeight)
-		[
-			SNew(SHorizontalBox)
-
-			+ SHorizontalBox::Slot()
-				.VAlign(VAlign_Center)
-				.Padding(5, 5)
-				.FillWidth(0.35f)
-				[
-					SNew(STextBlock)
-					.Text(LOCTEXT("CalibrateCameraPose", "Calibrate Nodal Offset"))
-					.ToolTipText(LOCTEXT("CalibrateNodalOffsetTooltip", "If checked, the final camera pose computed by the solver will be used to save the nodal offset to the Lens File."))
-				]
-
-			+ SHorizontalBox::Slot()
-				.VAlign(VAlign_Center)
-				.Padding(5, 5)
-				.FillWidth(0.65f)
-				[BuildCalibrateNodalOffsetWidget()]
-		]
-
 	+ SVerticalBox::Slot() // Calibration Rows
 		.AutoHeight()
 		.MaxHeight(12 * FCameraCalibrationWidgetHelpers::DefaultRowHeight)
@@ -811,7 +757,7 @@ bool UCameraLensDistortionAlgoCheckerboard::GetLensDistortion(
 
 	ECalibrationFlags SolverFlags = ECalibrationFlags::None;
 
-	if (bUseExtrinsicsGuess)
+	if (CVarUseExtrinsicsGuess.GetValueOnGameThread())
 	{
 		EnumAddFlags(SolverFlags, ECalibrationFlags::UseExtrinsicGuess);
 	}
@@ -875,7 +821,7 @@ bool UCameraLensDistortionAlgoCheckerboard::GetLensDistortion(
 	OutFocus = LastRow->CameraData.InputFocus;
 	OutZoom = LastRow->CameraData.InputZoom;
 
-	if (bCalibrateNodalOffset)
+	if (CVarUseExtrinsicsGuess.GetValueOnGameThread())
 	{
 		// See if the camera already had an offset applied, in which case we need to account for it.
 		FTransform ExistingOffset = FTransform::Identity;
