CameraCalibration: Fix filterable actor picker widget to display the actor label instead of the asset name
#rnx

diff --git a/Engine/Plugins/VirtualProduction/CameraCalibration/Source/CameraCalibrationEditor/Private/UI/SFilterableActorPicker.cpp b/Engine/Plugins/VirtualProduction/CameraCalibration/Source/CameraCalibrationEditor/Private/UI/SFilterableActorPicker.cpp
--- a/Engine/Plugins/VirtualProduction/CameraCalibration/Source/CameraCalibrationEditor/Private/UI/SFilterableActorPicker.cpp
+++ b/Engine/Plugins/VirtualProduction/CameraCalibration/Source/CameraCalibrationEditor/Private/UI/SFilterableActorPicker.cpp
@@ -197,6 +197,12 @@ FText SFilterableActorPicker::OnGetAssetName() const
 		return LOCTEXT("None", "None");
 	}
 
+	const AActor* Actor = Cast<AActor>(AssetData.GetAsset());
+	if (Actor)
+	{
+		return FText::AsCultureInvariant(Actor->GetActorNameOrLabel());
+	}
+
 	return FText::AsCultureInvariant(AssetData.AssetName.ToString());
 }
 
