CameraCalibration: Add automated tests for distortion and nodal offset calibration
#rnx

diff --git a/Engine/Plugins/Runtime/OpenCV/Source/OpenCVHelper/Private/OpenCVHelper.cpp b/Engine/Plugins/Runtime/OpenCV/Source/OpenCVHelper/Private/OpenCVHelper.cpp
--- a/Engine/Plugins/Runtime/OpenCV/Source/OpenCVHelper/Private/OpenCVHelper.cpp
+++ b/Engine/Plugins/Runtime/OpenCV/Source/OpenCVHelper/Private/OpenCVHelper.cpp
@@ -400,7 +400,7 @@ bool FOpenCVHelper::IdentifyArucoMarkers(TArray<FColor>& Image, FIntPoint ImageS
 	for (int32 MarkerIndex = 0; MarkerIndex < NumMarkers; ++MarkerIndex)
 	{
 		FArucoMarker NewMarker;
-		NewMarker.MarkedId = MarkerIds[MarkerIndex];
+		NewMarker.MarkerID = MarkerIds[MarkerIndex];
 
 		const std::vector<cv::Point2f>& MarkerCorners = Corners[MarkerIndex];
 
@@ -439,7 +439,7 @@ bool FOpenCVHelper::DrawArucoMarkers(const TArray<FArucoMarker>& Markers, UTextu
 	{
 		const FArucoMarker& Marker = Markers[MarkerIndex];
 
-		MarkerIds.push_back(Marker.MarkedId);
+		MarkerIds.push_back(Marker.MarkerID);
 
 		std::vector<cv::Point2f> MarkerCorners;
 		MarkerCorners.reserve(4);
@@ -617,6 +617,54 @@ bool FOpenCVHelper::SolvePnP(const TArray<FVector>& ObjectPoints, const TArray<F
 #endif // WITH_OPENCV
 }
 
+bool FOpenCVHelper::ProjectPoints(const TArray<FVector>& ObjectPoints, const FVector2D& FocalLength, const FVector2D& ImageCenter, const TArray<float>& DistortionParameters, const FTransform& CameraPose, TArray<FVector2f>& OutImagePoints)
+{
+#if WITH_OPENCV
+	const int32 NumPoints = ObjectPoints.Num();
+
+	// Convert from UE coordinates to OpenCV coordinates
+	TArray<FVector> CvObjectPoints;
+	CvObjectPoints.Reserve(NumPoints);
+
+	for (const FVector& ObjectPoint : ObjectPoints)
+	{
+		CvObjectPoints.Add(FOpenCVHelper::ConvertUnrealToOpenCV(ObjectPoint));
+	}
+
+	cv::Mat ObjectPointsMat = cv::Mat(NumPoints, 1, CV_64FC3, (void*)CvObjectPoints.GetData());
+
+	cv::Mat Rotation;
+	cv::Mat Translation;
+	FOpenCVHelper::MakeObjectVectorsFromCameraPose(CameraPose, Rotation, Translation);
+
+	// Initialize the camera matrix that will be used in each call to projectPoints()
+	cv::Mat CameraMatrix = cv::Mat::eye(3, 3, CV_64F);
+
+	CameraMatrix.at<double>(0, 0) = FocalLength.X;
+	CameraMatrix.at<double>(1, 1) = FocalLength.Y;
+	CameraMatrix.at<double>(0, 2) = ImageCenter.X;
+	CameraMatrix.at<double>(1, 2) = ImageCenter.Y;
+
+	cv::Mat DistortionParametersMat = cv::Mat(DistortionParameters.Num(), 1, CV_32FC1, (void*)DistortionParameters.GetData());
+
+	// cv::projectPoints requires that the 3D points and 2D points have the same bit depth, so we compute projected points with double precision, and then convert them back to floats before outputting
+	TArray<FVector2D> ImagePointsDoublePrecision;
+	ImagePointsDoublePrecision.Init(FVector2D(0.0, 0.0), NumPoints);
+	cv::Mat ImagePointsMat = cv::Mat(NumPoints, 1, CV_64FC2, (void*)ImagePointsDoublePrecision.GetData());
+
+	cv::projectPoints(ObjectPointsMat, Rotation, Translation, CameraMatrix, DistortionParametersMat, ImagePointsMat);
+
+	for (const FVector2D& Point : ImagePointsDoublePrecision)
+	{
+		OutImagePoints.Add(FVector2f(Point.X, Point.Y));
+	}
+
+	return true;
+#else
+	return false;
+#endif // WITH_OPENCV
+}
+
 bool FOpenCVHelper::FitLine3D(const TArray<FVector>& InPoints, FVector& OutLine, FVector& OutPointOnLine)
 {
 #if WITH_OPENCV
diff --git a/Engine/Plugins/Runtime/OpenCV/Source/OpenCVHelper/Public/OpenCVHelper.h b/Engine/Plugins/Runtime/OpenCV/Source/OpenCVHelper/Public/OpenCVHelper.h
--- a/Engine/Plugins/Runtime/OpenCV/Source/OpenCVHelper/Public/OpenCVHelper.h
+++ b/Engine/Plugins/Runtime/OpenCV/Source/OpenCVHelper/Public/OpenCVHelper.h
@@ -69,7 +69,7 @@ enum class EArucoDictionary : uint8
 
 struct OPENCVHELPER_API FArucoMarker
 {
-	int32 MarkedId;
+	int32 MarkerID = 0;
 	FVector2f Corners[4];
 };
 
@@ -158,6 +158,9 @@ public:
 	/** Compute the camera pose that minimizes the reprojection error of the input object points and image points */
 	static bool SolvePnP(const TArray<FVector>& ObjectPoints, const TArray<FVector2f>& ImagePoints, const FVector2D& FocalLength, const FVector2D& ImageCenter, const TArray<float>& DistortionParameters, FTransform& OutCameraPose);
 
+	/** Project the input object points to the 2D image plane defined by the input camera intrinsics and camera pose */
+	static bool ProjectPoints(const TArray<FVector>& ObjectPoints, const FVector2D& FocalLength, const FVector2D& ImageCenter, const TArray<float>& DistortionParameters, const FTransform& CameraPose, TArray<FVector2f>& OutImagePoints);
+
 	/** Find a 3D fit line that passes through the input points, as well as a point on that line */
 	static bool FitLine3D(const TArray<FVector>& InPoints, FVector& OutLine, FVector& OutPointOnLine);
 
diff --git a/Engine/Plugins/VirtualProduction/CameraCalibration/Source/CameraCalibrationEditor/CameraCalibrationEditor.Build.cs b/Engine/Plugins/VirtualProduction/CameraCalibration/Source/CameraCalibrationEditor/CameraCalibrationEditor.Build.cs
--- a/Engine/Plugins/VirtualProduction/CameraCalibration/Source/CameraCalibrationEditor/CameraCalibrationEditor.Build.cs
+++ b/Engine/Plugins/VirtualProduction/CameraCalibration/Source/CameraCalibrationEditor/CameraCalibrationEditor.Build.cs
@@ -22,6 +22,7 @@ namespace UnrealBuildTool.Rules
 					"CoreUObject",
 	                "CurveEditor",
 					"DesktopPlatform",
+					"DeveloperSettings",
 					"EditorStyle",
 					"EditorWidgets",
 					"Engine",
diff --git a/Engine/Plugins/VirtualProduction/CameraCalibration/Source/CameraCalibrationEditor/Private/Calibrators/CameraNodalOffsetAlgoAruco.cpp b/Engine/Plugins/VirtualProduction/CameraCalibration/Source/CameraCalibrationEditor/Private/Calibrators/CameraNodalOffsetAlgoAruco.cpp
--- a/Engine/Plugins/VirtualProduction/CameraCalibration/Source/CameraCalibrationEditor/Private/Calibrators/CameraNodalOffsetAlgoAruco.cpp
+++ b/Engine/Plugins/VirtualProduction/CameraCalibration/Source/CameraCalibrationEditor/Private/Calibrators/CameraNodalOffsetAlgoAruco.cpp
@@ -176,7 +176,7 @@ bool UCameraNodalOffsetAlgoAruco::PopulatePoints(FText& OutErrorMessage)
 		{
 			// Build calibrator point name based on the detected marker
 			const FString DictionaryName = UE::CameraCalibration::Private::GetArucoDictionaryName(Dictionary);
-			const FString MarkerName = FString::Printf(TEXT("%s-%d-%s"), *DictionaryName, Marker.MarkedId, *CornerNames[CornerIdx]); //-V557
+			const FString MarkerName = FString::Printf(TEXT("%s-%d-%s"), *DictionaryName, Marker.MarkerID, *CornerNames[CornerIdx]); //-V557
 
 			FNodalOffsetPointsCalibratorPointData PointCache;
 
diff --git a/Engine/Plugins/VirtualProduction/CameraCalibration/Source/CameraCalibrationEditor/Private/TestCameraCalibration.cpp b/Engine/Plugins/VirtualProduction/CameraCalibration/Source/CameraCalibrationEditor/Private/TestCameraCalibration.cpp
new file mode 100644
--- /dev/null
+++ b/Engine/Plugins/VirtualProduction/CameraCalibration/Source/CameraCalibrationEditor/Private/TestCameraCalibration.cpp
@@ -0,0 +1,487 @@
+// Copyright Epic Games, Inc. All Rights Reserved.
+
+
+#if WITH_DEV_AUTOMATION_TESTS
+
+#include "CoreMinimal.h"
+
+#include "Calibrators/CameraCalibrationSolver.h"
+#include "CameraCalibrationUtilsPrivate.h"
+#include "Dialog/SCustomDialog.h"
+#include "Engine/Engine.h"
+#include "Engine/Texture2D.h"
+#include "LensData.h"
+#include "Misc/AutomationTest.h"
+#include "Models/SphericalLensModel.h"
+#include "OpenCVHelper.h"
+#include "Slate/DeferredCleanupSlateBrush.h"
+#include "TestCameraCalibrationSettings.h"
+#include "Widgets/Images/SImage.h"
+
+IMPLEMENT_SIMPLE_AUTOMATION_TEST(FTestDistortionSpherical, "Plugin.CameraCalibration.TestDistortionSpherical", EAutomationTestFlags::EditorContext | EAutomationTestFlags::EngineFilter)
+IMPLEMENT_SIMPLE_AUTOMATION_TEST(FTestNodalOffset, "Plugin.CameraCalibration.TestNodalOffset", EAutomationTestFlags::EditorContext | EAutomationTestFlags::EngineFilter)
+
+namespace UE::Private::CameraCalibration::AutomatedTests
+{
+	static TSharedPtr<FDeferredCleanupSlateBrush> TextureBrush;
+
+	void TestDistortionCalibration(FAutomationTestBase& Test)
+	{
+		const UTestCameraCalibrationSettings* TestSettings = GetDefault<UTestCameraCalibrationSettings>();
+
+		// Extract the solver settings from the automated test settings
+		ECalibrationFlags SolverFlags = ECalibrationFlags::None;
+		if (TestSettings->bUseCameraIntrinsicGuess)
+		{
+			EnumAddFlags(SolverFlags, ECalibrationFlags::UseIntrinsicGuess);
+		}
+		if (TestSettings->bUseCameraExtrinsicGuess)
+		{
+			EnumAddFlags(SolverFlags, ECalibrationFlags::UseExtrinsicGuess);
+		}
+		if (TestSettings->bFixFocalLength)
+		{
+			EnumAddFlags(SolverFlags, ECalibrationFlags::FixFocalLength);
+		}
+		if (TestSettings->bFixImageCenter)
+		{
+			EnumAddFlags(SolverFlags, ECalibrationFlags::FixPrincipalPoint);
+		}
+		if (TestSettings->bFixExtrinsics)
+		{
+			EnumAddFlags(SolverFlags, ECalibrationFlags::FixExtrinsics);
+		}
+		if (TestSettings->bFixZeroDistortion)
+		{
+			EnumAddFlags(SolverFlags, ECalibrationFlags::FixZeroDistortion);
+		}
+
+		/** 
+		 * Step 1: Establish the ground-truth 3D object point and 2D image point data
+		 */
+
+		 // Initialize the ground-truth image properties
+		const FIntPoint ImageSize = TestSettings->ImageSize;
+		const float SensorWidth = TestSettings->SensorDimensions.X;
+
+		// Initialize the ground-truth camera intrinsics
+		const double TrueFocalLength = TestSettings->FocalLength;
+		const double TrueFocalLengthPixels = (TrueFocalLength / SensorWidth) * ImageSize.X;
+		const FVector2D TrueFxFy = FVector2D(TrueFocalLengthPixels, TrueFocalLengthPixels);
+
+		const FVector2D ImageCenter = FVector2D((ImageSize.X - 1) * 0.5, (ImageSize.Y - 1) * 0.5);
+
+		// Initialize the ground-truth camera extrinsics
+		const FTransform TrueCameraPose = TestSettings->CameraTransform;
+
+		// Initialize the ground-truth distortion parameters
+		const FSphericalDistortionParameters SphericalParams = TestSettings->SphericalDistortionParameters;
+		const TArray<float> DistortionParameters = { SphericalParams.K1, SphericalParams.K2, SphericalParams.P1, SphericalParams.P2, SphericalParams.K3 };
+
+		// Initialize the 3D object points. For a distortion calibration, these points represent co-planar points that would be found on a real calibrator at a reasonable distance from the physical camera.
+		// In the current implementation, they simulate the corners of a checkerboard with size and dimensions defined in the test settings.
+		// In the future, this test could be expanded to support simulating an aruco or charuco board, or other known patterns.
+		TArray<TArray<FVector>> ObjectPoints;
+		TArray<TArray<FVector2f>> ImagePoints;
+		TArray<FTransform> TrueCameraPoses;
+		TArray<FTransform> EstimatedCameraPoses;
+
+		TArray<FVector> Points3d;
+
+		// We assume that the checkerboard lies perfectly in the YZ plane (same axes as our image plane)
+		// To compute the 3D points, we place the board at a fixed distance from the camera (set by the corresponding test setting).
+		// We then place the center of the board at (X, 0, 0), and find the top left intersection of checkerboard squares (not the top left corner of the checkerboard). 
+		// The points start at the top left corner and proceed to the right, then start over for the next row. This matches how OpenCV would order the detected corners of a checkerboard.
+		const double DistanceFromCamera = TestSettings->CalibratorDistanceFromCamera;
+		const double SquareSize = TestSettings->CheckerboardSquareSize;
+
+		// The dimensions defined in the test settings are the number of squares in the board, however we only care about the number of intersections, so we reduce each dimension by 1.
+		const FIntPoint CheckerboardCornerDimensions = TestSettings->CheckerboardDimensions - FIntPoint(1, 1);
+
+		// The top left intersection is one square length to the right and down from the actual top left corner of the board
+		const double TopLeftCornerX = -(TestSettings->CheckerboardDimensions.X * 0.5 * SquareSize) + SquareSize;
+		const double TopLeftCornerY =  (TestSettings->CheckerboardDimensions.Y * 0.5 * SquareSize) - SquareSize;
+
+		for (int32 RowIndex = 0; RowIndex < CheckerboardCornerDimensions.Y; ++RowIndex)
+		{
+			for (int32 ColumnIndex = 0; ColumnIndex < CheckerboardCornerDimensions.X; ++ColumnIndex)
+			{
+				const double Y = TopLeftCornerX + (SquareSize * ColumnIndex);
+				const double Z = TopLeftCornerY - (SquareSize * RowIndex);
+				Points3d.Add(FVector(DistanceFromCamera, Y, Z));
+			}
+		}
+
+		// In a real calibration, there are likely to be images of the checkerboard taken from multiple camera angles
+		// The test settings have a setting for the number of camera views to use.
+		// The current strategy is to only rotate the camera a maximum of 30 degrees to the left and to the right, and translate it in space to keep the board in view.
+		// The number of views, therefore, determines how far to move the camera to generate each view
+		double StartRotation = 0.0;
+		double StartTranslation = 0.0;
+		double RotationStep = 0.0;
+		double TranslationStep = 0.0;
+		if (TestSettings->NumCameraViews > 1)
+		{
+			StartRotation = 30.0;
+			StartTranslation = TestSettings->CalibratorDistanceFromCamera * 0.6;
+
+			RotationStep =    (StartRotation    * 2) / (TestSettings->NumCameraViews - 1);
+			TranslationStep = (StartTranslation * 2) / (TestSettings->NumCameraViews - 1);
+		}
+
+		// For each of the camera views, project the 3D calibrator points to the image plane
+		for (int32 ViewIndex = 0; ViewIndex < TestSettings->NumCameraViews; ++ViewIndex)
+		{
+			const FVector ViewTranslation = FVector(0, StartTranslation - (TranslationStep * ViewIndex), 0);
+			const FRotator ViewRotation = FRotator(0, -StartRotation + (RotationStep * ViewIndex), 0);
+
+			FTransform CameraMotion = FTransform::Identity;
+			CameraMotion.SetTranslation(ViewTranslation);
+			CameraMotion.SetRotation(ViewRotation.Quaternion());
+
+			FTransform CameraPoseForView = CameraMotion * TrueCameraPose;
+			FTransform EstimatedCameraPoseForView = CameraMotion * TestSettings->EstimatedCameraTransform;
+
+			TArray<FVector2f> Points2d;
+			bool bResult = FOpenCVHelper::ProjectPoints(Points3d, TrueFxFy, ImageCenter, DistortionParameters, CameraPoseForView, Points2d);
+			if (!bResult)
+			{
+				Test.AddError(TEXT("Project Points failed. Test could not be completed"));
+				return;
+			}
+
+			ObjectPoints.Add(Points3d);
+			ImagePoints.Add(Points2d);
+			TrueCameraPoses.Add(CameraPoseForView);
+			EstimatedCameraPoses.Add(EstimatedCameraPoseForView);
+		}
+
+		// Debug setting that draws all of the projected 2D checkerboard patterns to an image that can be displayed to the user to validate that the points make sense (and are within the bounds of the image)
+		if (TestSettings->bShowCheckerboardImage)
+ 		{
+			UTexture2D* DebugTexture = UTexture2D::CreateTransient(ImageSize.X, ImageSize.Y, EPixelFormat::PF_B8G8R8A8);
+			UE::CameraCalibration::Private::ClearTexture(DebugTexture, FColor::Black);
+
+			for (const TArray<FVector2f>& Image : ImagePoints)
+			{
+				FOpenCVHelper::DrawCheckerboardCorners(Image, CheckerboardCornerDimensions, DebugTexture);
+			}
+
+			TextureBrush = FDeferredCleanupSlateBrush::CreateBrush(DebugTexture);
+
+			TSharedRef<SCustomDialog> DetectionWindow =
+				SNew(SCustomDialog)
+				.UseScrollBox(false)
+				.Content()
+				[
+					SNew(SImage)
+					.Image(TextureBrush->GetSlateBrush())
+				]
+				.Buttons
+				({
+					SCustomDialog::FButton(NSLOCTEXT("TestCameraCalibration", "OkButton", "Ok")),
+				});
+
+			DetectionWindow->Show();
+ 		}
+
+		/**
+		 * Step 2: Introduce errors into the 3D and 2D point data to simulate real-world inaccuracies that occur when doing calibration
+		 */
+
+		// Introduce some random noise to the 3D points. The checkerboard is a rigid object, so the individual 3D positions of each corner cannot change randomly with respect to one another.
+		// However, the entire board could have the wrong pose if, for example, the tracking data is noisy, or if the tracked rigid-body pose sent to UE from the tracking system is not precise.
+		TArray<TArray<FVector>> NoisyObjectPoints;
+		NoisyObjectPoints.Reserve(ObjectPoints.Num());
+		for (const TArray<FVector>& Object : ObjectPoints)
+		{
+			TArray<FVector> NoisyPoints;
+			NoisyPoints.Reserve(Object.Num());
+
+			const double NoiseScale = TestSettings->ObjectPointNoiseScale;
+			const double NoiseX = (FMath::SRand() - 0.5) * (NoiseScale * 2);
+			const double NoiseY = (FMath::SRand() - 0.5) * (NoiseScale * 2);
+			const double NoiseZ = (FMath::SRand() - 0.5) * (NoiseScale * 2);
+
+			for (const FVector& Point : Object)
+			{
+				NoisyPoints.Add(Point + FVector(NoiseX, NoiseY, NoiseZ));
+			}
+
+			NoisyObjectPoints.Add(NoisyPoints);
+		}
+
+		// Introduce some random noise to the 2D points. This simulates poor checkerboard detection, which could occur if the checkerboard is not perfectly in-focus, if the image resolution is too low,
+		// or if there is some other imprecision in the corner detection algorithm. 
+		TArray<TArray<FVector2f>> NoisyImagePoints;
+		NoisyImagePoints.Reserve(ImagePoints.Num());
+		for (const TArray<FVector2f>& Image : ImagePoints)
+		{
+			TArray<FVector2f> NoisyPoints;
+			NoisyPoints.Reserve(Image.Num());
+
+			// Unlike the 3D points, the 2D image points could all be randomly noisy compared to one another
+			for (const FVector2f& Point : Image)
+			{
+				const double NoiseScale = TestSettings->ImagePointNoiseScale;
+				const double NoiseX = (FMath::SRand() - 0.5) * (NoiseScale * 2);
+				const double NoiseY = (FMath::SRand() - 0.5) * (NoiseScale * 2);
+
+				NoisyPoints.Add(Point + FVector2f(NoiseX, NoiseY));
+			}
+
+			NoisyImagePoints.Add(NoisyPoints);
+		}
+
+		/**
+		 * Step 3: Run the calibration solver to compute the focal length, image center, distortion parameters, and camera poses for each view.
+		 * If no errors were introduced into the data, the expectation is that the solver will be able to compute the ground-truth for all of these properties.
+		 * If this is not the case, then we either uncover bugs in the solver, or learn more about the limitations of the solver
+		 * The introduction of errors should reveal how real-world calibrations can produce poor results if the quality of the input data is poor.
+		 */
+
+		// This is fixed until the tests support testing anamorphic calibration
+		constexpr float PixelAspect = 1.0f;
+
+		FVector2D CalibratedFxFy = TrueFxFy;
+		if (TestSettings->bUseCameraIntrinsicGuess)
+		{
+			const double EstimatedFocalLength = TestSettings->EstimatedFocalLength;
+			const double EstimatedFocalLengthPixels = (EstimatedFocalLength / SensorWidth) * ImageSize.X;
+			CalibratedFxFy = FVector2D(EstimatedFocalLengthPixels, EstimatedFocalLengthPixels);
+		}
+
+		FVector2D CalibratedImageCenter = ImageCenter;
+		TArray<float> CalibratedDistortionParameters;
+
+		const double Error = FCameraCalibrationSolver::CalibrateCamera(
+			USphericalLensModel::StaticClass(),
+			NoisyObjectPoints,
+			NoisyImagePoints,
+			ImageSize,
+			CalibratedFxFy,
+			CalibratedImageCenter,
+			CalibratedDistortionParameters,
+			EstimatedCameraPoses,
+			PixelAspect,
+			SolverFlags
+		);
+
+		/**
+		 * Step 4: Output the test results
+		 */
+
+		Test.AddInfo(TEXT("Ground-Truth Image Properties:"));
+		Test.AddInfo(FString::Printf(TEXT("\t\tImage Dimensions: (%d, %d)"), TestSettings->ImageSize.X, TestSettings->ImageSize.Y));
+		Test.AddInfo(FString::Printf(TEXT("\t\tSensor Dimensions: (%f, %f)"), TestSettings->SensorDimensions.X, TestSettings->SensorDimensions.Y));
+
+		Test.AddInfo(TEXT("Ground-Truth Camera Intrinsics:"));
+		Test.AddInfo(FString::Printf(TEXT("\t\tFocal Length: %lf mm (%lf pixels)"), TrueFocalLength, TrueFxFy.X));
+		Test.AddInfo(FString::Printf(TEXT("\t\tImage Center: (%lf, %lf)"), ImageCenter.X, ImageCenter.Y));
+
+		Test.AddInfo(TEXT("Ground-Truth Distortion Coefficients:"));
+		Test.AddInfo(FString::Printf(TEXT("\t\tK1: %f"), SphericalParams.K1));
+		Test.AddInfo(FString::Printf(TEXT("\t\tK2: %f"), SphericalParams.K2));
+		Test.AddInfo(FString::Printf(TEXT("\t\tK3: %f"), SphericalParams.K3));
+		Test.AddInfo(FString::Printf(TEXT("\t\tP1: %f"), SphericalParams.P1));
+		Test.AddInfo(FString::Printf(TEXT("\t\tP2: %f"), SphericalParams.P2));
+
+		Test.AddInfo(TEXT("Ground-Truth Camera Poses:"));
+		for (int32 PoseIndex = 0; PoseIndex < TrueCameraPoses.Num(); ++PoseIndex)
+		{
+			const FTransform& Pose = TrueCameraPoses[PoseIndex];
+
+			Test.AddInfo(FString::Printf(TEXT("\t\tPose %d"), PoseIndex));
+
+			const FVector Translation = Pose.GetTranslation();
+			const FRotator Rotator = Pose.GetRotation().Rotator();
+
+			Test.AddInfo(FString::Printf(TEXT("\t\t\t\tTranslation: (%lf, %lf, %lf)"), Translation.X, Translation.Y, Translation.Z));
+			Test.AddInfo(FString::Printf(TEXT("\t\t\t\tRotation:    (%lf, %lf, %lf)"), Rotator.Roll, Rotator.Pitch, Rotator.Yaw));
+		}
+
+		Test.AddInfo(TEXT("\n"));
+ 		Test.AddInfo(FString::Printf(TEXT("Result RMS Error: %lf"), Error));
+		Test.AddInfo(TEXT("\n"));
+
+		Test.AddInfo(TEXT("Calibrated Camera Intrinsics:"));
+ 		Test.AddInfo(FString::Printf(TEXT("\t\tFocal Length: %lf mm (%lf pixels)"), (CalibratedFxFy.X / ImageSize.X)* SensorWidth, CalibratedFxFy.X));
+		Test.AddInfo(FString::Printf(TEXT("\t\tImage Center: (%lf, %lf)"), CalibratedImageCenter.X, CalibratedImageCenter.Y));
+
+ 		Test.AddInfo(TEXT("Calibrated Distortion Coefficients:"));
+		Test.AddInfo(FString::Printf(TEXT("\t\tK1: %f"), CalibratedDistortionParameters[0]));
+		Test.AddInfo(FString::Printf(TEXT("\t\tK2: %f"), CalibratedDistortionParameters[1]));
+		Test.AddInfo(FString::Printf(TEXT("\t\tK3: %f"), CalibratedDistortionParameters[2]));
+		Test.AddInfo(FString::Printf(TEXT("\t\tP1: %f"), CalibratedDistortionParameters[3]));
+		Test.AddInfo(FString::Printf(TEXT("\t\tP2: %f"), CalibratedDistortionParameters[4]));
+
+		Test.AddInfo(TEXT("Calibrated Camera Poses:"));
+		for (int32 PoseIndex = 0; PoseIndex < EstimatedCameraPoses.Num(); ++PoseIndex)
+		{
+			const FTransform& Pose = EstimatedCameraPoses[PoseIndex];
+
+			Test.AddInfo(FString::Printf(TEXT("\t\tPose %d"), PoseIndex));
+
+			const FVector Translation = Pose.GetTranslation();
+			const FRotator Rotator = Pose.GetRotation().Rotator();
+
+			Test.AddInfo(FString::Printf(TEXT("\t\t\t\tTranslation: (%lf, %lf, %lf)"), Translation.X, Translation.Y, Translation.Z));
+			Test.AddInfo(FString::Printf(TEXT("\t\t\t\tRotation:    (%lf, %lf, %lf)"), Rotator.Roll, Rotator.Pitch, Rotator.Yaw));
+		}
+	}
+
+	void TestNodalOffsetCalibration(FAutomationTestBase& Test)
+	{
+		const UTestCameraCalibrationSettings* TestSettings = GetDefault<UTestCameraCalibrationSettings>();
+
+		/**
+		 * Step 1: Establish the ground-truth 3D object point and 2D image point data
+		 */
+		
+		// Initialize the ground-truth image properties
+		const FIntPoint ImageSize = TestSettings->ImageSize;
+		const float SensorWidth = TestSettings->SensorDimensions.X;
+
+		// Initialize the ground-truth camera intrinsics
+		const double TrueFocalLength = TestSettings->FocalLength;
+		const double TrueFocalLengthPixels = (TrueFocalLength / SensorWidth) * ImageSize.X;
+		const FVector2D TrueFxFy = FVector2D(TrueFocalLengthPixels, TrueFocalLengthPixels);
+
+		const FVector2D ImageCenter = FVector2D((ImageSize.X - 1) * 0.5, (ImageSize.Y - 1) * 0.5);
+
+		// Initialize the ground-truth camera extrinsics
+		const FTransform TrueCameraPose = TestSettings->CameraTransform;
+
+		// Initialize the ground-truth distortion parameters
+		const FSphericalDistortionParameters SphericalParams = TestSettings->SphericalDistortionParameters;
+		const TArray<float> DistortionParameters = { SphericalParams.K1, SphericalParams.K2, SphericalParams.P1, SphericalParams.P2, SphericalParams.K3 };
+
+		// Initialize the 3D object points. For a distortion calibration, these points represent co-planar points that would be found on a real calibrator at a reasonable distance from the physical camera.
+		// In the current implementation, the camera is fixed at the world origin (0, 0, 0).
+		// These object points represent an object that is a known distance from the camera and that might come from a calibration object that is 120cm x 60cm (~4ft x ~2ft)
+		const double DistanceFromCamera = TestSettings->CalibratorDistanceFromCamera;
+		const double SquareSize = TestSettings->CheckerboardSquareSize;
+
+		// The dimensions defined in the test settings are the number of squares in the board, however we only care about the number of intersections, so we reduce each dimension by 1.
+		const FIntPoint CheckerboardCornerDimensions = TestSettings->CheckerboardDimensions - FIntPoint(1, 1);
+
+		// The top left intersection is one square length to the right and down from the actual top left corner of the board
+		const double TopLeftCornerX = -(TestSettings->CheckerboardDimensions.X * 0.5 * SquareSize) + SquareSize;
+		const double TopLeftCornerY = (TestSettings->CheckerboardDimensions.Y * 0.5 * SquareSize) - SquareSize;
+
+		TArray<FVector> ObjectPoints;
+		for (int32 RowIndex = 0; RowIndex < CheckerboardCornerDimensions.Y; ++RowIndex)
+		{
+			for (int32 ColumnIndex = 0; ColumnIndex < CheckerboardCornerDimensions.X; ++ColumnIndex)
+			{
+				const double Y = TopLeftCornerX + (SquareSize * ColumnIndex);
+				const double Z = TopLeftCornerY - (SquareSize * RowIndex);
+				ObjectPoints.Add(FVector(DistanceFromCamera, Y, Z));
+			}
+		}
+
+		TArray<FVector2f> ImagePoints;
+		bool bResult = FOpenCVHelper::ProjectPoints(ObjectPoints, TrueFxFy, ImageCenter, DistortionParameters, TrueCameraPose, ImagePoints);
+		if (!bResult)
+		{
+			Test.AddError(TEXT("Project Points failed. Test could not be completed"));
+			return;
+		}
+
+		/**
+		 * Step 2: Run SolvePnP to solve for the camera pose using perfect input data. The expectation is that the solver will be able to compute the ground-truth camera pose.
+		 */
+
+		FTransform PerfectCameraPoseResult;
+		FOpenCVHelper::SolvePnP(ObjectPoints, ImagePoints, TrueFxFy, ImageCenter, DistortionParameters, PerfectCameraPoseResult);
+
+		/**
+		 * Step 3: Introduce errors into the 3D and 2D point data to simulate real-world inaccuracies that occur when doing calibration
+		 */
+
+		 // Introduce some random noise to the 3D points. The checkerboard is a rigid object, so the individual 3D positions of each corner cannot change randomly with respect to one another.
+		 // However, the entire board could have the wrong pose if, for example, the tracking data is noisy, or if the tracked rigid-body pose sent to UE from the tracking system is not precise.
+		TArray<FVector> NoisyObjectPoints;
+		NoisyObjectPoints.Reserve(ObjectPoints.Num());
+		{
+			const double NoiseScale = TestSettings->ObjectPointNoiseScale;
+			const double NoiseX = (FMath::SRand() - 0.5) * (NoiseScale * 2);
+			const double NoiseY = (FMath::SRand() - 0.5) * (NoiseScale * 2);
+			const double NoiseZ = (FMath::SRand() - 0.5) * (NoiseScale * 2);
+
+			for (const FVector& Point : ObjectPoints)
+			{
+				NoisyObjectPoints.Add(Point + FVector(NoiseX, NoiseY, NoiseZ));
+			}
+		}
+
+		// Introduce some random noise to the 2D points. This simulates poor checkerboard detection, which could occur if the checkerboard is not perfectly in-focus, if the image resolution is too low,
+		// or if there is some other imprecision in the corner detection algorithm. 
+		TArray<FVector2f> NoisyImagePoints;
+		NoisyImagePoints.Reserve(ImagePoints.Num());
+		for (const FVector2f& Point : ImagePoints)
+		{
+			const double NoiseScale = TestSettings->ImagePointNoiseScale;
+			const double NoiseX = (FMath::SRand() - 0.5) * (NoiseScale * 2);
+			const double NoiseY = (FMath::SRand() - 0.5) * (NoiseScale * 2);
+
+			NoisyImagePoints.Add(Point + FVector2f(NoiseX, NoiseY));
+		}
+
+		/**
+		 * Step 4: Run SolvePnP to solve for the camera pose using imperfect input data, including noisy data and an incorrect guess for focal length
+		 */
+
+		const double EstimatedFocalLength = TestSettings->EstimatedFocalLength;
+		const double EstimatedFocalLengthPixels = (EstimatedFocalLength / SensorWidth) * ImageSize.X;
+		const FVector2D EstimatedFxFy = FVector2D(EstimatedFocalLengthPixels, EstimatedFocalLengthPixels);
+
+		FTransform ImperfectCameraPoseResult;
+		FOpenCVHelper::SolvePnP(NoisyObjectPoints, NoisyImagePoints, EstimatedFxFy, ImageCenter, DistortionParameters, ImperfectCameraPoseResult);
+
+		/**
+		 * Step 4: Output the test results
+		 */
+
+		{
+			Test.AddInfo(TEXT("Ground-Truth Camera Pose"));
+			const FVector Translation = TrueCameraPose.GetTranslation();
+			const FRotator Rotator = TrueCameraPose.GetRotation().Rotator();
+
+			Test.AddInfo(FString::Printf(TEXT("\t\t\t\tTranslation: (%lf, %lf, %lf)"), Translation.X, Translation.Y, Translation.Z));
+			Test.AddInfo(FString::Printf(TEXT("\t\t\t\tRotation:    (%lf, %lf, %lf)"), Rotator.Roll, Rotator.Pitch, Rotator.Yaw));
+		}
+
+		{
+			Test.AddInfo(TEXT("Perfectly Solved Camera Pose"));
+			const FVector Translation = PerfectCameraPoseResult.GetTranslation();
+			const FRotator Rotator = PerfectCameraPoseResult.GetRotation().Rotator();
+
+			Test.AddInfo(FString::Printf(TEXT("\t\t\t\tTranslation: (%lf, %lf, %lf)"), Translation.X, Translation.Y, Translation.Z));
+			Test.AddInfo(FString::Printf(TEXT("\t\t\t\tRotation:    (%lf, %lf, %lf)"), Rotator.Roll, Rotator.Pitch, Rotator.Yaw));
+		}
+
+		{
+			Test.AddInfo(TEXT("Imperfectly Solved Camera Pose"));
+			const FVector Translation = ImperfectCameraPoseResult.GetTranslation();
+			const FRotator Rotator = ImperfectCameraPoseResult.GetRotation().Rotator();
+
+			Test.AddInfo(FString::Printf(TEXT("\t\t\t\tTranslation: (%lf, %lf, %lf)"), Translation.X, Translation.Y, Translation.Z));
+			Test.AddInfo(FString::Printf(TEXT("\t\t\t\tRotation:    (%lf, %lf, %lf)"), Rotator.Roll, Rotator.Pitch, Rotator.Yaw));
+		}
+	}
+}
+
+bool FTestDistortionSpherical::RunTest(const FString& Parameters)
+{
+	UE::Private::CameraCalibration::AutomatedTests::TestDistortionCalibration(*this);
+	return true;
+}
+
+bool FTestNodalOffset::RunTest(const FString& Parameters)
+{
+	UE::Private::CameraCalibration::AutomatedTests::TestNodalOffsetCalibration(*this);
+	return true;
+}
+
+#endif // WITH_DEV_AUTOMATION_TESTS
diff --git a/Engine/Plugins/VirtualProduction/CameraCalibration/Source/CameraCalibrationEditor/Private/TestCameraCalibrationSettings.cpp b/Engine/Plugins/VirtualProduction/CameraCalibration/Source/CameraCalibrationEditor/Private/TestCameraCalibrationSettings.cpp
new file mode 100644
--- /dev/null
+++ b/Engine/Plugins/VirtualProduction/CameraCalibration/Source/CameraCalibrationEditor/Private/TestCameraCalibrationSettings.cpp
@@ -0,0 +1,23 @@
+// Copyright Epic Games, Inc. All Rights Reserved.
+#include "TestCameraCalibrationSettings.h"
+
+UTestCameraCalibrationSettings::UTestCameraCalibrationSettings() 
+{
+}
+
+FName UTestCameraCalibrationSettings::GetCategoryName() const 
+{ 
+	return TEXT("Plugins"); 
+}
+
+#if WITH_EDITOR
+FText UTestCameraCalibrationSettings::GetSectionText() const 
+{ 
+	return NSLOCTEXT("TestCameraCalibrationSettings", "TestCameraCalibrationSettingsSection", "Camera Calibration Test Settings"); 
+}
+
+FName UTestCameraCalibrationSettings::GetSectionName() const 
+{	
+	return TEXT("Camera Calibration Test"); 
+}
+#endif
diff --git a/Engine/Plugins/VirtualProduction/CameraCalibration/Source/CameraCalibrationEditor/Private/TestCameraCalibrationSettings.h b/Engine/Plugins/VirtualProduction/CameraCalibration/Source/CameraCalibrationEditor/Private/TestCameraCalibrationSettings.h
new file mode 100644
--- /dev/null
+++ b/Engine/Plugins/VirtualProduction/CameraCalibration/Source/CameraCalibrationEditor/Private/TestCameraCalibrationSettings.h
@@ -0,0 +1,113 @@
+// Copyright Epic Games, Inc. All Rights Reserved.
+#pragma  once
+
+#include "Engine/DeveloperSettings.h"
+#include "Models/SphericalLensModel.h"
+
+#include "TestCameraCalibrationSettings.generated.h"
+
+/**
+ * Settings for the camera calibration testing
+ */
+UCLASS(config = EditorPerProjectUserSettings)
+class UTestCameraCalibrationSettings : public UDeveloperSettings
+{
+	GENERATED_BODY()
+
+public:
+	UTestCameraCalibrationSettings();
+
+	//~ Begin UDeveloperSettings interface
+	virtual FName GetCategoryName() const;
+#if WITH_EDITOR
+	virtual FText GetSectionText() const override;
+	virtual FName GetSectionName() const override;
+#endif
+	//~ End UDeveloperSettings interface
+
+public:
+
+#if WITH_EDITORONLY_DATA
+
+	/** Size of the image that the test camera would generate */
+	UPROPERTY(config, EditAnywhere, Category = "True Camera Properties")
+	FIntPoint ImageSize = FIntPoint(1920, 1080);
+
+	/** Size of the sensor (in mm) of the test camera */
+	UPROPERTY(config, EditAnywhere, Category = "True Camera Properties")
+	FVector2f SensorDimensions = FVector2f(23.76f, 13.365f);
+
+	/** True 3D pose of the test camera */
+	UPROPERTY(config, EditAnywhere, Category = "True Camera Properties")
+	FTransform CameraTransform = FTransform::Identity;
+
+	/** True focal length of the test camera */
+	UPROPERTY(config, EditAnywhere, Category = "True Camera Properties")
+	double FocalLength = 35.0;
+
+	/** True distortion coefficients of the test lens */
+	UPROPERTY(config, EditAnywhere, Category = "True Camera Properties")
+	FSphericalDistortionParameters SphericalDistortionParameters;
+
+	/** Number of camera views to generate images from */
+	UPROPERTY(config, EditAnywhere, Category = "True Camera Properties")
+	int32 NumCameraViews = 1;
+
+	/** Number of squares in the checkerboard (columns x rows) */
+	UPROPERTY(config, EditAnywhere, Category = "Calibrator Properties")
+	FIntPoint CheckerboardDimensions = FIntPoint(16, 9);
+
+	/** Size (in cm) of the length of each square of the checkerboard */
+	UPROPERTY(config, EditAnywhere, Category = "Calibrator Properties")
+	double CheckerboardSquareSize = 5.0;
+
+	/** Distance (in cm) to place the test calibrator from the test camera */
+	UPROPERTY(config, EditAnywhere, Category = "Calibrator Properties")
+	double CalibratorDistanceFromCamera = 200.0;
+
+	/** Maximum amount of random noise (in cm) to add to 3D object points */
+	UPROPERTY(config, EditAnywhere, Category = "Input Data Errors")
+	double ObjectPointNoiseScale = 0.0;
+
+	/** Maximum amount of random noise (in pixels) to add to the 2D image points */
+	UPROPERTY(config, EditAnywhere, Category = "Input Data Errors")
+	double ImagePointNoiseScale = 0.0;
+
+	/** If true, the solver will use the estimated focal length as an initial guess for the calibrated focal length */
+	UPROPERTY(config, EditAnywhere, Category = "Calibrate Camera Solver Settings")
+	bool bUseCameraIntrinsicGuess = false;
+
+	/** User-provided guess for the focal length */
+	UPROPERTY(config, EditAnywhere, Category = "Calibrate Camera Solver Settings", meta = (EditCondition = "bUseCameraIntrinsicGuess == true"))
+	double EstimatedFocalLength = 35.0;
+
+	/** If true, the solver will use the estimate camera transform as an initial guess for the calibrated camera pose */
+	UPROPERTY(config, EditAnywhere, Category = "Calibrate Camera Solver Settings")
+	bool bUseCameraExtrinsicGuess = false;
+
+	/** User-provided guess for the camera pose */
+	UPROPERTY(config, EditAnywhere, Category = "Calibrate Camera Solver Settings", meta = (EditCondition = "bUseCameraExtrinsicGuess == true"))
+	FTransform EstimatedCameraTransform = FTransform::Identity;
+
+	/** If true, the solver will not try to solve for focal length, and will use the focal length estimate */
+	UPROPERTY(config, EditAnywhere, Category = "Calibrate Camera Solver Settings")
+	bool bFixFocalLength = false;
+
+	/** If true, the solver will not try to solve for image center, and will use (0.5, 0.5) */
+	UPROPERTY(config, EditAnywhere, Category = "Calibrate Camera Solver Settings")
+	bool bFixImageCenter = false;
+
+	/** If true, the solver will not try to solve for the camera pose, and will use the camera transform estimate */
+	UPROPERTY(config, EditAnywhere, Category = "Calibrate Camera Solver Settings")
+	bool bFixExtrinsics = false;
+
+	/** If true, the solver will not try to solve for distortion, and will assume there is zero distortion in the image */
+	UPROPERTY(config, EditAnywhere, Category = "Calibrate Camera Solver Settings")
+	bool bFixZeroDistortion = false;
+
+	/** If true, a debug image will be displayed showing the 2D projected points of the test checkerboard for each of the camera views */
+	UPROPERTY(config, EditAnywhere, Category = "Debug")
+	bool bShowCheckerboardImage = false;
+
+#endif
+};
