CameraCalibration: Add Lens Distortion algo that operates on aruco patterns.

diff --git a/Engine/Plugins/VirtualProduction/CameraCalibration/Source/CameraCalibrationEditor/Private/AssetEditor/CameraCalibrationStepsController.cpp b/Engine/Plugins/VirtualProduction/CameraCalibration/Source/CameraCalibrationEditor/Private/AssetEditor/CameraCalibrationStepsController.cpp
--- a/Engine/Plugins/VirtualProduction/CameraCalibration/Source/CameraCalibrationEditor/Private/AssetEditor/CameraCalibrationStepsController.cpp
+++ b/Engine/Plugins/VirtualProduction/CameraCalibration/Source/CameraCalibrationEditor/Private/AssetEditor/CameraCalibrationStepsController.cpp
@@ -1441,6 +1441,7 @@ void FCameraCalibrationStepsController::TogglePlay()
 	}
 	else
 	{
+		// TODO: Trigger the current step (and ultimately the algo) to cache any data it cares about (like 3D scene data)
 		MediaPlayer->Pause();
 	}
 }
@@ -1683,7 +1684,7 @@ bool FCameraCalibrationStepsController::CalculateNormalizedMouseClickPosition(co
 	return true;
 }
 
-bool FCameraCalibrationStepsController::ReadMediaPixels(TArray<FColor>& Pixels, FIntPoint& Size, ETextureRenderTargetFormat& PixelFormat, FText& OutErrorMessage, ESimulcamViewportPortion ViewportPortion) const
+bool FCameraCalibrationStepsController::ReadMediaPixels(TArray<FColor>& Pixels, FIntPoint& Size, FText& OutErrorMessage, ESimulcamViewportPortion ViewportPortion) const
 {
 	// Get the media plate texture render target 2d
 
@@ -1702,7 +1703,11 @@ bool FCameraCalibrationStepsController::ReadMediaPixels(TArray<FColor>& Pixels,
 		return false;
 	}
 
-	PixelFormat = MediaPlateRenderTarget->RenderTargetFormat;
+	if (MediaPlateRenderTarget->RenderTargetFormat != ETextureRenderTargetFormat::RTF_RGBA8)
+	{
+		OutErrorMessage = LOCTEXT("InvalidFormat", "MediaPlateRenderTarget did not have the expected RTF_RGBA8 format");
+		return false;
+	}
 
 	// Read the pixels onto CPU
 	TArray<FColor> MediaPixels;
diff --git a/Engine/Plugins/VirtualProduction/CameraCalibration/Source/CameraCalibrationEditor/Private/AssetEditor/CameraCalibrationStepsController.h b/Engine/Plugins/VirtualProduction/CameraCalibration/Source/CameraCalibrationEditor/Private/AssetEditor/CameraCalibrationStepsController.h
--- a/Engine/Plugins/VirtualProduction/CameraCalibration/Source/CameraCalibrationEditor/Private/AssetEditor/CameraCalibrationStepsController.h
+++ b/Engine/Plugins/VirtualProduction/CameraCalibration/Source/CameraCalibrationEditor/Private/AssetEditor/CameraCalibrationStepsController.h
@@ -136,7 +136,7 @@ public:
 	UWorld* GetWorld() const;
 
 	/** Reads the pixels in the media plate */
-	bool ReadMediaPixels(TArray<FColor>& Pixels, FIntPoint& Size, ETextureRenderTargetFormat& PixelFormat, FText& OutErrorMessage, ESimulcamViewportPortion ViewportPortion = ESimulcamViewportPortion::FullViewport) const;
+	bool ReadMediaPixels(TArray<FColor>& Pixels, FIntPoint& Size, FText& OutErrorMessage, ESimulcamViewportPortion ViewportPortion = ESimulcamViewportPortion::FullViewport) const;
 
 	/** Returns true if the overlay transform pass is currently enabled */
 	bool IsOverlayEnabled(EOverlayPassType OverlayPass = EOverlayPassType::ToolOverlay) const;
diff --git a/Engine/Plugins/VirtualProduction/CameraCalibration/Source/CameraCalibrationEditor/Private/Calibrators/CameraLensDistortionAlgoAruco.cpp b/Engine/Plugins/VirtualProduction/CameraCalibration/Source/CameraCalibrationEditor/Private/Calibrators/CameraLensDistortionAlgoAruco.cpp
new file mode 100644
--- /dev/null
+++ b/Engine/Plugins/VirtualProduction/CameraCalibration/Source/CameraCalibrationEditor/Private/Calibrators/CameraLensDistortionAlgoAruco.cpp
@@ -0,0 +1,1083 @@
+// Copyright Epic Games, Inc. All Rights Reserved.
+
+#include "CameraLensDistortionAlgoAruco.h"
+
+#include "AssetEditor/CameraCalibrationStepsController.h"
+#include "AssetEditor/LensDistortionTool.h"
+#include "CalibrationPointComponent.h"
+#include "Camera/CameraActor.h"
+#include "CameraCalibrationEditorLog.h"
+#include "CameraCalibrationSolver.h"
+#include "CameraCalibrationUtilsPrivate.h"
+#include "Dialog/SCustomDialog.h"
+#include "Dom/JsonObject.h"
+#include "Engine/Texture2D.h"
+#include "Framework/MultiBox/MultiBoxBuilder.h"
+#include "ImageUtils.h"
+#include "JsonObjectConverter.h"
+#include "LensComponent.h"
+#include "LensFile.h"
+#include "Misc/MessageDialog.h"
+#include "Models/AnamorphicLensModel.h"
+#include "Models/SphericalLensModel.h"
+#include "OpenCVHelper.h"
+#include "UI/CameraCalibrationWidgetHelpers.h"
+#include "Widgets/Input/SButton.h"
+#include "Widgets/Input/SCheckBox.h"
+#include "Widgets/Input/SComboBox.h"
+#include "Widgets/Input/SNumericEntryBox.h"
+#include "Widgets/Views/SListView.h"
+#include "Widgets/Views/STableRow.h"
+#include "UI/SFilterableActorPicker.h"
+#include "UI/SImageTexture.h"
+
+
+#define LOCTEXT_NAMESPACE "CameraLensDistortionAlgoAruco"
+
+#if WITH_EDITOR
+static TAutoConsoleVariable<bool> CVarUseIntrinsicsGuess(TEXT("LensDistortionAruco.UseIntrinsicsGuess"), true, TEXT("If true, the solver initializes the camera intrinsics to a user-provided estimate. Otherwise, the solver will compute the initial values."));
+static TAutoConsoleVariable<bool> CVarFixExtrinsics(TEXT("LensDistortionAruco.FixExtrinsics"), false, TEXT("If true, the solver will fix the camera extrinsics to the user-provided camera poses"));
+static TAutoConsoleVariable<bool> CVarFixZeroDistortion(TEXT("LensDistortionAruco.FixZeroDistortion"), false, TEXT("If true, the solver will fix all distortion values to always be 0"));
+static TAutoConsoleVariable<bool> CVarUseExtrinsicsGuess(TEXT("LensDistortionAruco.UseExtrinsicsGuess"), false, TEXT("If true, the actual calibrator and camera poses will be used when running the solver"));
+#endif
+
+const int UCameraLensDistortionAlgoAruco::DATASET_VERSION = 1;
+
+// String constants for import/export
+namespace UE::CameraCalibration::Private::LensDistortionArucoExportFields
+{
+	static const FString Version(TEXT("Version"));
+}
+
+namespace UE::CameraCalibration::Private
+{
+	// Test the input array of weak pointers for validity, and return a list of valid pointers to calibration point components
+	void GetValidActiveComponents(const TArray<TWeakObjectPtr<UCalibrationPointComponent>>& WeakComponents, TArray<UCalibrationPointComponent*>& ValidComponents)
+	{
+		ValidComponents.Reserve(WeakComponents.Num());
+		for (const TWeakObjectPtr<UCalibrationPointComponent>& WeakComponent : WeakComponents)
+		{
+			if (UCalibrationPointComponent* Component = WeakComponent.Get())
+			{
+				ValidComponents.Add(Component);
+			}
+		}
+	}
+
+	// Generate a debug texture with the aruco points from the input row drawn onto it
+	UTexture2D* CreateDebugTexture(TSharedPtr<FLensDistortionArucoRowData> Row)
+	{
+		if (!Row)
+		{
+			return nullptr;
+		}
+
+		TArray<FArucoMarker> DebugMarkers;
+		DebugMarkers.Reserve(Row->ArucoPoints.Num());
+
+		// Extract the data from the aruco points that OpenCV expects to draw the debug markers
+		for (int32 MarkerIndex = 0; MarkerIndex < Row->ArucoPoints.Num(); ++MarkerIndex)
+		{
+			const FArucoCalibrationPoint& Marker = Row->ArucoPoints[MarkerIndex];
+
+			FArucoMarker NewMarker;
+			NewMarker.MarkerID = Marker.MarkerID;
+			for (int32 CornerIndex = 0; CornerIndex < 4; ++CornerIndex)
+			{
+				NewMarker.Corners[CornerIndex] = Marker.Corners2D[CornerIndex];
+			}
+
+			DebugMarkers.Add(NewMarker);
+		}
+
+		UTexture2D* DebugTexture = FImageUtils::CreateTexture2DFromImage(Row->MediaImage);
+		FOpenCVHelper::DrawArucoMarkers(DebugMarkers, DebugTexture);
+
+		return DebugTexture;
+	}
+
+	// Open a dialog to display the input image 
+	void DisplayDebugImageDialog(UTexture2D* Image)
+	{
+		if (Image)
+		{
+			TSharedRef<SCustomDialog> DebugImageDialog =
+				SNew(SCustomDialog)
+				.UseScrollBox(false)
+				.Content()
+				[
+					SNew(SImageTexture, Image)
+				]
+				.Buttons
+				({
+					SCustomDialog::FButton(LOCTEXT("OkButton", "Ok")),
+				});
+
+			DebugImageDialog->Show();
+		}
+	}
+}
+
+namespace CameraLensDistortionAlgoAruco
+{
+	class SCalibrationRowGenerator : public SMultiColumnTableRow<TSharedPtr<FLensDistortionArucoRowData>>
+	{
+
+	public:
+		SLATE_BEGIN_ARGS(SCalibrationRowGenerator) {}
+
+		/** The list item for this row */
+		SLATE_ARGUMENT(TSharedPtr<FLensDistortionArucoRowData>, CalibrationRowData)
+
+		SLATE_END_ARGS()
+
+		void Construct(const FArguments& Args, const TSharedRef<STableViewBase>& OwnerTableView)
+		{
+			CalibrationRowData = Args._CalibrationRowData;
+
+			SMultiColumnTableRow<TSharedPtr<FLensDistortionArucoRowData>>::Construct(
+				FSuperRowType::FArguments()
+				.Padding(1.0f),
+				OwnerTableView
+			);
+		}
+
+		//~Begin SMultiColumnTableRow
+		virtual TSharedRef<SWidget> GenerateWidgetForColumn(const FName& ColumnName) override
+		{
+			if (!CalibrationRowData)
+			{
+				return SNullWidget::NullWidget;
+			}
+
+			if (ColumnName == TEXT("Index"))
+			{
+				return SNew(SVerticalBox)
+					+ SVerticalBox::Slot()
+					.HAlign(EHorizontalAlignment::HAlign_Center)
+					.VAlign(EVerticalAlignment::VAlign_Center)
+					[SNew(STextBlock).Text(FText::AsNumber(CalibrationRowData->Index))];
+			}
+
+			if (ColumnName == TEXT("NumMarkers"))
+			{
+				return SNew(SVerticalBox)
+					+ SVerticalBox::Slot()
+					.HAlign(EHorizontalAlignment::HAlign_Center)
+					.VAlign(EVerticalAlignment::VAlign_Center)
+					[SNew(STextBlock).Text(FText::AsNumber(CalibrationRowData->ArucoPoints.Num()))];
+			}
+
+			if (ColumnName == TEXT("Image"))
+			{
+				// Generate a transient thumbnail texture to display in the tool
+				const FImage& MediaImage = CalibrationRowData->MediaImage;
+
+				constexpr int32 ResolutionDivider = 4;
+				const FIntPoint ThumbnailSize = FIntPoint(MediaImage.SizeX / ResolutionDivider, MediaImage.SizeY / ResolutionDivider);
+
+				FImage ThumbnailImage;
+				FImageCore::ResizeTo(MediaImage, ThumbnailImage, ThumbnailSize.X, ThumbnailSize.Y, MediaImage.Format, MediaImage.GetGammaSpace());
+
+				if (UTexture2D* Thumbnail = FImageUtils::CreateTexture2DFromImage(ThumbnailImage))
+				{
+					return SNew(SImageTexture, Thumbnail)
+						.MinDesiredHeight(4 * FCameraCalibrationWidgetHelpers::DefaultRowHeight)
+						.MaxDesiredHeight(4 * FCameraCalibrationWidgetHelpers::DefaultRowHeight);
+				}
+				else
+				{
+					const FString Text = FString::Printf(TEXT("Image Unavailable"));
+					return SNew(STextBlock).Text(FText::FromString(Text));
+				}
+			}
+
+			return SNullWidget::NullWidget;
+		}
+		//~End SMultiColumnTableRow
+
+
+	private:
+		TSharedPtr<FLensDistortionArucoRowData> CalibrationRowData;
+	};
+
+};
+
+void UCameraLensDistortionAlgoAruco::Initialize(ULensDistortionTool* InTool)
+{
+	WeakTool = InTool;
+
+	// Find all actors in the current level that have calibration point components
+	TArray<AActor*> CalibratorActors;
+	UE::CameraCalibration::Private::FindActorsWithCalibrationComponents(CalibratorActors);
+
+	// If at least one actor was found, set the initial calibrator to the first one in the list
+	if (CalibratorActors.Num() > 0)
+	{
+		SetCalibrator(CalibratorActors[0]);
+	}
+}
+
+void UCameraLensDistortionAlgoAruco::Shutdown()
+{
+	WeakTool.Reset();
+	WeakCalibrator.Reset();
+}
+
+bool UCameraLensDistortionAlgoAruco::SupportsModel(const TSubclassOf<ULensModel>& LensModel) const
+{
+	return ((LensModel == USphericalLensModel::StaticClass()) || (LensModel == UAnamorphicLensModel::StaticClass()));
+}
+
+void UCameraLensDistortionAlgoAruco::Tick(float DeltaTime)
+{
+}
+
+bool UCameraLensDistortionAlgoAruco::OnViewportClicked(const FGeometry& MyGeometry, const FPointerEvent& MouseEvent)
+{
+	// We only respond to left clicks
+	if (MouseEvent.GetEffectingButton() != EKeys::LeftMouseButton)
+	{
+		return false;
+	}
+
+	// Attempt to add a new calibration row with data from the current media image, and display an error dialog to the user if any errors occur
+	FText ErrorMessage;
+	if (!AddCalibrationRow(ErrorMessage))
+	{
+		if (ErrorMessage.IsEmpty())
+		{
+			ErrorMessage = LOCTEXT("UnknownCalibrationError", "An error occurred. Please check the output log for more details.");
+		}
+
+		FMessageDialog::Open(EAppMsgType::Ok, ErrorMessage);
+	}
+
+	// If the media was currently paused, then resume it after gathering data from the current image
+	if (ULensDistortionTool* LensDistortionTool = WeakTool.Get())
+	{
+		if (FCameraCalibrationStepsController* StepsController = LensDistortionTool->GetCameraCalibrationStepsController())
+		{
+			StepsController->Play();
+		}
+	}
+
+	return true;
+}
+
+bool UCameraLensDistortionAlgoAruco::GetLensDistortion(
+	float& OutFocus,
+	float& OutZoom,
+	FDistortionInfo& OutDistortionInfo,
+	FFocalLengthInfo& OutFocalLengthInfo,
+	FImageCenterInfo& OutImageCenterInfo,
+	TSubclassOf<ULensModel>& OutLensModel,
+	double& OutError,
+	FText& OutErrorMessage)
+{
+	if (CalibrationRows.Num() < 1)
+	{
+		OutErrorMessage = LOCTEXT("NotEnoughSamples", "At least 1 calibration row is required");
+		return false;
+	}
+
+	ULensDistortionTool* LensDistortionTool = WeakTool.Get();
+	if (!LensDistortionTool)
+	{
+		return false;
+	}
+
+	FCameraCalibrationStepsController* StepsController = LensDistortionTool->GetCameraCalibrationStepsController();
+	if (!StepsController)
+	{
+		return false;
+	}
+
+	ULensFile* LensFile = StepsController->GetLensFile();
+	if (!LensFile)
+	{
+		OutErrorMessage = LOCTEXT("LensFileNotFound", "LensFile not found");
+		return false;
+	}
+
+	// Only parameters data mode supported at the moment
+	if (LensFile->DataMode != ELensDataMode::Parameters)
+	{
+		OutErrorMessage = LOCTEXT("OnlyParametersDataModeSupported", "Only Parameters Data Mode supported");
+		return false;
+	}
+
+	// Initialize the focal length estimate in pixels
+	const FIntPoint ImageSize = LensFile->CameraFeedInfo.GetDimensions();
+
+	const float PixelAspect = LensFile->LensInfo.SqueezeFactor;
+
+	if (FMath::IsNearlyZero(PixelAspect))
+	{
+		OutErrorMessage = LOCTEXT("PixelAspectNearlyZero", "The pixel aspect in the Lens Information is zero, which is invalid.");
+		return false;
+	}
+
+	const float PhysicalSensorWidth = StepsController->GetLensFileEvaluationInputs().Filmback.SensorWidth;
+	const float DesqueezeSensorWidth = PhysicalSensorWidth * PixelAspect;
+
+	if (FMath::IsNearlyZero(DesqueezeSensorWidth))
+	{
+		OutErrorMessage = LOCTEXT("SensorWidthNearlyZero", "The sensor width of the CineCamera is zero, which is invalid.");
+		return false;
+	}
+
+	if (!FocalLengthEstimate.IsSet())
+	{
+		OutErrorMessage = LOCTEXT("NoFocalLengthEstimateSet", "Please enter a value (in mm) for the focal length estimate.");
+		return false;
+	}
+
+	const float FocalLengthEstimateValue = FocalLengthEstimate.GetValue();
+	const double Fx = (FocalLengthEstimateValue / DesqueezeSensorWidth) * ImageSize.X;
+
+	// When operating on a desqueezed image, we expect our pixel aspect to be square, so horizontal and vertical field of view are assumed to be equal (i.e. Fx == Fy)
+	FVector2D FocalLength = FVector2D(Fx, Fx);
+	FVector2D ImageCenter = FVector2D((ImageSize.X - 1) * 0.5, (ImageSize.Y - 1) * 0.5);
+
+	TArray<TArray<FVector>> Samples3d;
+	Samples3d.Reserve(CalibrationRows.Num());
+
+	TArray<TArray<FVector2f>> Samples2d;
+	Samples2d.Reserve(CalibrationRows.Num());
+
+	TArray<FTransform> CameraPoses;
+	CameraPoses.Reserve(CalibrationRows.Num());
+
+	// Extract the 3D points, 2D points, and camera poses from each row to pass to the solver
+	for (const TSharedPtr<FLensDistortionArucoRowData>& Row : CalibrationRows)
+	{
+		check(Row.IsValid());
+
+		TArray<FVector> Points3d;
+		Points3d.Reserve(Row->ArucoPoints.Num() * 4);
+
+		TArray<FVector2f> Points2d;
+		Points2d.Reserve(Row->ArucoPoints.Num() * 4);
+
+		for (const FArucoCalibrationPoint& Marker : Row->ArucoPoints)
+		{
+			for (int32 CornerIndex = 0; CornerIndex < 4; ++CornerIndex)
+			{
+				Points3d.Add(Marker.Corners3D[CornerIndex]);
+				Points2d.Add(Marker.Corners2D[CornerIndex]);
+			}
+		}
+
+		Samples3d.Add(Points3d);
+		Samples2d.Add(Points2d);
+
+		CameraPoses.Add(Row->CameraPose);
+	}
+
+	ECalibrationFlags SolverFlags = ECalibrationFlags::None;
+
+	if (CVarUseExtrinsicsGuess.GetValueOnGameThread())
+	{
+		EnumAddFlags(SolverFlags, ECalibrationFlags::UseExtrinsicGuess);
+	}
+
+	if (CVarUseIntrinsicsGuess.GetValueOnAnyThread())
+	{
+		EnumAddFlags(SolverFlags, ECalibrationFlags::UseIntrinsicGuess);
+	}
+
+	if (CVarFixExtrinsics.GetValueOnAnyThread())
+	{
+		EnumAddFlags(SolverFlags, ECalibrationFlags::FixExtrinsics);
+	}
+
+	if (CVarFixZeroDistortion.GetValueOnAnyThread())
+	{
+		EnumAddFlags(SolverFlags, ECalibrationFlags::FixZeroDistortion);
+	}
+
+	if (bFixFocalLength)
+	{
+		EnumAddFlags(SolverFlags, ECalibrationFlags::FixFocalLength);
+	}
+
+	if (bFixImageCenter)
+	{
+		EnumAddFlags(SolverFlags, ECalibrationFlags::FixPrincipalPoint);
+	}
+
+	TArray<float> DistortionCoefficients;
+
+	OutError = FCameraCalibrationSolver::CalibrateCamera(
+		LensFile->LensInfo.LensModel,
+		Samples3d,
+		Samples2d,
+		ImageSize,
+		FocalLength,
+		ImageCenter,
+		DistortionCoefficients,
+		CameraPoses,
+		PixelAspect,
+		SolverFlags
+	);
+
+	OutLensModel = LensFile->LensInfo.LensModel;
+	OutDistortionInfo.Parameters = DistortionCoefficients;
+
+	OutFocalLengthInfo.FxFy = FocalLength / ImageSize;
+	OutImageCenterInfo.PrincipalPoint = ImageCenter / ImageSize;
+
+	const FLensFileEvaluationInputs LensFileEvalInputs = StepsController->GetLensFileEvaluationInputs();
+	OutFocus = LensFileEvalInputs.Focus;
+	OutZoom = LensFileEvalInputs.Zoom;
+
+	if (CVarUseExtrinsicsGuess.GetValueOnGameThread())
+	{
+		check(CameraPoses.Num() == CalibrationRows.Num());
+
+		// See if the camera already had an offset applied, in which case we need to account for it.
+		FTransform ExistingOffset = FTransform::Identity;
+
+		bool bWasNodalOffsetApplied = false;
+		if (const ULensComponent* LensComponent = StepsController->FindLensComponent())
+		{
+			bWasNodalOffsetApplied = LensComponent->WasNodalOffsetAppliedThisTick();
+		}
+
+		if (bWasNodalOffsetApplied)
+		{
+			FNodalPointOffset NodalPointOffset;
+
+			if (LensFile->EvaluateNodalPointOffset(OutFocus, OutZoom, NodalPointOffset))
+			{
+				ExistingOffset.SetTranslation(NodalPointOffset.LocationOffset);
+				ExistingOffset.SetRotation(NodalPointOffset.RotationOffset);
+			}
+		}
+
+		TArray<FNodalPointOffset> NewNodalOffsets;
+		NewNodalOffsets.Reserve(CameraPoses.Num());
+		
+		for (int RowIndex = 0; RowIndex < CameraPoses.Num(); ++RowIndex)
+		{
+			const FTransform DesiredCameraTransform = CameraPoses[RowIndex];
+			FTransform DesiredOffset = DesiredCameraTransform * CalibrationRows[RowIndex]->CameraPose.Inverse() * ExistingOffset;
+
+			FNodalPointOffset NewNodalOffset;
+			NewNodalOffset.LocationOffset = DesiredOffset.GetLocation();
+			NewNodalOffset.RotationOffset = DesiredOffset.GetRotation();
+
+			NewNodalOffsets.Add(NewNodalOffset);
+		}
+
+		LensFile->AddNodalOffsetPoint(OutFocus, OutZoom, NewNodalOffsets[0]);
+	}
+
+	return true;
+}
+
+void UCameraLensDistortionAlgoAruco::OnDistortionSavedToLens()
+{
+	// Since the calibration result was saved, there is no further use for the current samples.
+	ClearCalibrationRows();
+}
+
+bool UCameraLensDistortionAlgoAruco::AddCalibrationRow(FText& OutErrorMessage)
+{
+	// Detect the aruco dictionary to use by looking at the names of the calibration points on the selected calibrator actor
+	AActor* CalibratorActor = WeakCalibrator.Get();
+	if (!CalibratorActor)
+	{
+		OutErrorMessage = LOCTEXT("InvalidCalibrator", "Please pick a calibrator actor from the list that contains aruco calibration points.");
+		return false;
+	}
+
+	EArucoDictionary Dictionary = UE::CameraCalibration::Private::GetArucoDictionaryForCalibrator(CalibratorActor);
+	if (Dictionary == EArucoDictionary::None)
+	{
+		OutErrorMessage = LOCTEXT("InvalidArucoDictionary", "The names of the calibration point component(s) of the selected calibrator do not match any known aruco dictionaries.");
+		return false;
+	}
+
+	ULensDistortionTool* LensDistortionTool = WeakTool.Get();
+	if (!LensDistortionTool)
+	{
+		UE_LOG(LogCameraCalibrationEditor, Warning, TEXT("Failed to gather calibration data from the viewport. The Lens Distortion Tool was invalid."));
+		return false;
+	}
+
+	FCameraCalibrationStepsController* StepsController = LensDistortionTool->GetCameraCalibrationStepsController();
+	if (!StepsController)
+	{
+		UE_LOG(LogCameraCalibrationEditor, Warning, TEXT("Failed to gather calibration data from the viewport. The Steps Controller was invalid."));
+		return false;
+	}
+
+	// Read the image data from the current media source into a CPU-accessible array
+	TArray<FColor> Pixels;
+	FIntPoint Size = FIntPoint(0, 0);
+	if (!StepsController->ReadMediaPixels(Pixels, Size, OutErrorMessage, ESimulcamViewportPortion::CameraFeed))
+	{
+		return false;
+	}
+
+	// Identify any aruco markers matching the current dictionary in the media image
+	TArray<FArucoMarker> IdentifiedMarkers;
+	bool bResult = FOpenCVHelper::IdentifyArucoMarkers(Pixels, Size, Dictionary, IdentifiedMarkers);
+
+	if (!bResult || IdentifiedMarkers.IsEmpty())
+	{
+		OutErrorMessage = LOCTEXT("NoArucoMarkersFound", "No aruco markers could be identified in the media image, or they did not match the aruco dictionary of the selected calibrator actor.");
+		return false;
+	}
+
+	// For each identified marker, search the calibration components to find the subpoints matching that marker and the 3D location of each of its corners
+	TArray<FArucoCalibrationPoint> ArucoCalibrationPoints;
+	ArucoCalibrationPoints.Reserve(IdentifiedMarkers.Num());
+
+	TArray<UCalibrationPointComponent*> ValidActiveCalibrationComponents;
+	UE::CameraCalibration::Private::GetValidActiveComponents(ActiveCalibratorComponents, ValidActiveCalibrationComponents);
+
+	for (const FArucoMarker& Marker : IdentifiedMarkers)
+	{
+		FArucoCalibrationPoint ArucoCalibrationPoint;
+		if (UE::CameraCalibration::Private::FindArucoCalibrationPoint(ValidActiveCalibrationComponents, Dictionary, Marker, ArucoCalibrationPoint))
+		{
+			ArucoCalibrationPoints.Add(ArucoCalibrationPoint);
+		}
+	}
+
+	// In the current implementation, we expect each set of 3D-2D point correspondences to be coplanar. All of the aruco markers belonging to a single calibration component should be coplanar.
+	// Therefore, we group the aruco calibration points by the calibration component they belong to, and each distinct component will contribute its own set of coplanar points.
+	TMap<UObject*, TArray<FArucoCalibrationPoint>> MarkerSets;
+	UE::CameraCalibration::Private::SortArucoCalibrationPoints(ArucoCalibrationPoints, MarkerSets);
+
+	ExportSessionData();
+
+	// For each set of coplanar markers, add a new calibration row
+	for (const TPair<UObject*, TArray<FArucoCalibrationPoint>>& MarkerSetPair : MarkerSets)
+	{
+		TSharedPtr<FLensDistortionArucoRowData> NewRow = MakeShared<FLensDistortionArucoRowData>();
+		NewRow->Index = LensDistortionTool->AdvanceSessionRowIndex();
+
+		// Save the aruco points associated with this row 
+		NewRow->ArucoPoints = MarkerSetPair.Value;
+
+		// Get the current pose of the selected CineCamera actor
+		NewRow->CameraPose = FTransform::Identity;
+		if (const ACameraActor* Camera = StepsController->GetCamera())
+		{
+			if (const UCameraComponent* CameraComponent = Camera->GetCameraComponent())
+			{
+				NewRow->CameraPose = CameraComponent->GetComponentToWorld();
+			}
+		}
+
+		// Save an image view of the captured frame
+		FImageView ImageView = FImageView(Pixels.GetData(), Size.X, Size.Y, ERawImageFormat::BGRA8);
+		ImageView.CopyTo(NewRow->MediaImage);
+
+		CalibrationRows.Add(NewRow);
+
+		// Export the data for this row to a .json file on disk
+		ExportRow(NewRow);
+	}
+
+	// Notify the ListView of the new data
+	if (CalibrationRowListWidget)
+	{
+		CalibrationRowListWidget->RequestListRefresh();
+	}
+
+	return true;
+}
+
+void UCameraLensDistortionAlgoAruco::ClearCalibrationRows()
+{
+	CalibrationRows.Empty();
+
+	if (CalibrationRowListWidget)
+	{
+		CalibrationRowListWidget->RequestListRefresh();
+	}
+	
+	// End the current calibration session (a new one will begin the next time a new row is added)
+	if (ULensDistortionTool* LenDistortionTool = WeakTool.Get())
+	{
+		LenDistortionTool->EndCalibrationSession();
+	}
+}
+
+void UCameraLensDistortionAlgoAruco::SetCalibrator(AActor* InCalibrator)
+{
+	WeakCalibrator = InCalibrator;
+
+	if (!InCalibrator)
+	{
+		return;
+	}
+
+	// Find all of the calibration components of the input calibrator 
+	constexpr uint32 NumInlineAllocations = 32;
+	TArray<UCalibrationPointComponent*, TInlineAllocator<NumInlineAllocations>> CalibrationPoints;
+	InCalibrator->GetComponents(CalibrationPoints);
+
+	ActiveCalibratorComponents.Empty();
+
+	// When a new calibrator is selected, all of its components will default to active
+	for (UCalibrationPointComponent* CalibrationPoint : CalibrationPoints)
+	{
+		// Only add calibration point components that have an attached scene component
+		if (CalibrationPoint && CalibrationPoint->GetAttachParent())
+		{
+			ActiveCalibratorComponents.Add(CalibrationPoint);
+		}
+	}
+}
+
+void UCameraLensDistortionAlgoAruco::ExportSessionData()
+{
+	using namespace UE::CameraCalibration::Private;
+
+	if (ULensDistortionTool* LensDistortionTool = WeakTool.Get())
+	{
+		// Add all data to a json object that is needed to run this algorithm that is NOT part of a specific row
+		TSharedPtr<FJsonObject> JsonSessionData = MakeShared<FJsonObject>();
+
+		JsonSessionData->SetNumberField(LensDistortionArucoExportFields::Version, DATASET_VERSION);
+
+		// Export the session data to a .json file
+		LensDistortionTool->ExportSessionData(JsonSessionData.ToSharedRef());
+	}
+}
+
+void UCameraLensDistortionAlgoAruco::ExportRow(TSharedPtr<FLensDistortionArucoRowData> Row)
+{
+	if (ULensDistortionTool* LensDistortionTool = WeakTool.Get())
+	{
+		if (const TSharedPtr<FJsonObject>& RowObject = FJsonObjectConverter::UStructToJsonObject<FLensDistortionArucoRowData>(Row.ToSharedRef().Get()))
+		{
+			// Export the row data to a .json file
+			LensDistortionTool->ExportCalibrationRow(Row->Index, RowObject.ToSharedRef(), Row->MediaImage);
+		}
+	}
+}
+
+bool UCameraLensDistortionAlgoAruco::HasCalibrationData() const
+{
+	return (CalibrationRows.Num() > 0);
+}
+
+void UCameraLensDistortionAlgoAruco::PreImportCalibrationData()
+{
+	// Clear the current set of rows before importing new ones
+	CalibrationRows.Empty();
+}
+
+int32 UCameraLensDistortionAlgoAruco::ImportCalibrationRow(const TSharedRef<FJsonObject>& CalibrationRowObject, const FImage& RowImage)
+{
+	// Create a new row to populate with data from the Json object
+	TSharedPtr<FLensDistortionArucoRowData> NewRow = MakeShared<FLensDistortionArucoRowData>();
+
+	// We enforce strict mode to ensure that every field in the UStruct of row data is present in the imported json.
+	// If any fields are missing, it is likely the row will be invalid, which will lead to errors in the calibration.
+	constexpr int64 CheckFlags = 0;
+	constexpr int64 SkipFlags = 0;
+	constexpr bool bStrictMode = true;
+	if (FJsonObjectConverter::JsonObjectToUStruct<FLensDistortionArucoRowData>(CalibrationRowObject, NewRow.Get(), CheckFlags, SkipFlags, bStrictMode))
+	{
+		if (!RowImage.RawData.IsEmpty())
+		{
+			NewRow->MediaImage = RowImage;
+		}
+
+		CalibrationRows.Add(NewRow);
+	}
+	else
+	{
+		UE_LOG(LogCameraCalibrationEditor, Warning, TEXT("LensDistortionAruco algo failed to import calibration row because at least one field could not be deserialized from the json file."));
+	}
+
+	return NewRow->Index;
+}
+
+void UCameraLensDistortionAlgoAruco::PostImportCalibrationData()
+{
+	// Sort imported calibration rows by row index
+	CalibrationRows.Sort([](const TSharedPtr<FLensDistortionArucoRowData>& LHS, const TSharedPtr<FLensDistortionArucoRowData>& RHS) { return LHS->Index < RHS->Index; });
+
+	// Notify the ListView of the new data
+	if (CalibrationRowListWidget)
+	{
+		CalibrationRowListWidget->RequestListRefresh();
+	}
+}
+
+TSharedRef<SWidget> UCameraLensDistortionAlgoAruco::BuildUI()
+{
+	return SNew(SVerticalBox)
+
+		+ SVerticalBox::Slot() // Calibrator picker
+		.VAlign(EVerticalAlignment::VAlign_Top)
+		.AutoHeight()
+		.MaxHeight(FCameraCalibrationWidgetHelpers::DefaultRowHeight)
+		[FCameraCalibrationWidgetHelpers::BuildLabelWidgetPair(LOCTEXT("Aruco", "Aruco"), BuildCalibratorPickerWidget())]
+
+		+ SVerticalBox::Slot() // Calibrator component picker
+		.VAlign(EVerticalAlignment::VAlign_Top)
+		.AutoHeight()
+		.MaxHeight(FCameraCalibrationWidgetHelpers::DefaultRowHeight)
+		[FCameraCalibrationWidgetHelpers::BuildLabelWidgetPair(LOCTEXT("CalibratorComponents", "Calibrator Component(s)"), BuildCalibrationComponentPickerWidget())]
+
+		+ SVerticalBox::Slot() // Solver Settings Title
+		.Padding(0, 5)
+		.AutoHeight()
+		.VAlign(EVerticalAlignment::VAlign_Center)
+		[
+			SNew(SBox) // Constrain the height
+			.MinDesiredHeight(FCameraCalibrationWidgetHelpers::DefaultRowHeight)
+			.MaxDesiredHeight(FCameraCalibrationWidgetHelpers::DefaultRowHeight)
+			[
+				SNew(SBorder) // Background color of title
+				.BorderImage(FAppStyle::GetBrush("DetailsView.CategoryTop"))
+				.BorderBackgroundColor(FLinearColor::White)
+				.VAlign(EVerticalAlignment::VAlign_Center)
+				[
+					SNew(SOverlay)
+					+ SOverlay::Slot() // Used to add left padding to the title
+					.Padding(5, 0, 0, 0)
+					[
+						SNew(STextBlock) // Title text
+						.Text(LOCTEXT("SolverSettings", "Solver Settings"))
+						.TransformPolicy(ETextTransformPolicy::ToUpper)
+						.Font(FAppStyle::Get().GetFontStyle(TEXT("PropertyWindow.BoldFont")))
+						.TextStyle(FAppStyle::Get(), "DetailsView.CategoryTextStyle")
+					]
+				]
+			]
+		]
+
+		+ SVerticalBox::Slot() // Focal Length Estimate
+		.VAlign(EVerticalAlignment::VAlign_Top)
+		.AutoHeight()
+		.MaxHeight(FCameraCalibrationWidgetHelpers::DefaultRowHeight)
+		[FCameraCalibrationWidgetHelpers::BuildLabelWidgetPair(LOCTEXT("FocalLengthEstimate", "Focal Length Estimate"), BuildFocalLengthEstimateWidget())]
+
+		+ SVerticalBox::Slot() // Fix Image Center
+		.VAlign(EVerticalAlignment::VAlign_Top)
+		.AutoHeight()
+		.MaxHeight(FCameraCalibrationWidgetHelpers::DefaultRowHeight)
+		[FCameraCalibrationWidgetHelpers::BuildLabelWidgetPair(LOCTEXT("FixImageCenter", "Fix Image Center"), BuildFixImageCenterWidget())]
+
+		+ SVerticalBox::Slot() // Calibration Rows
+		.AutoHeight()
+		.MaxHeight(12 * FCameraCalibrationWidgetHelpers::DefaultRowHeight)
+		[BuildCalibrationRowListWidget()]
+
+		+ SVerticalBox::Slot() // Action buttons (e.g. Clear)
+		.HAlign(EHorizontalAlignment::HAlign_Center)
+		.AutoHeight()
+		.Padding(0, 20)
+		[BuildCalibrationActionButtons()];
+}
+
+TSharedRef<SWidget> UCameraLensDistortionAlgoAruco::BuildCalibratorPickerWidget()
+{
+	return SNew(SFilterableActorPicker)
+		.OnSetObject_UObject(this, &UCameraLensDistortionAlgoAruco::OnCalibratorSelected)
+		.OnShouldFilterAsset_UObject(this, &UCameraLensDistortionAlgoAruco::DoesAssetHaveCalibrationComponent)
+		.ActorAssetData_UObject(this, &UCameraLensDistortionAlgoAruco::GetCalibratorAssetData);
+}
+
+TSharedRef<SWidget> UCameraLensDistortionAlgoAruco::BuildCalibrationComponentPickerWidget()
+{
+	return SNew(SHorizontalBox)
+
+		+ SHorizontalBox::Slot()
+		.AutoWidth()
+		.Padding(FMargin(0.0f, 0.0f, 0.0f, 0.0f))
+		[
+			SNew(SComboButton)
+			.OnGetMenuContent_UObject(this, &UCameraLensDistortionAlgoAruco::BuildCalibrationComponentMenu)
+			.ContentPadding(FMargin(4.0, 2.0))
+			.ButtonContent()
+			[
+				SNew(STextBlock).Text_UObject(this, &UCameraLensDistortionAlgoAruco::GetCalibrationComponentMenuText)
+			]
+		];
+}
+
+TSharedRef<SWidget> UCameraLensDistortionAlgoAruco::BuildCalibrationComponentMenu()
+{
+	// Generate menu
+	FMenuBuilder MenuBuilder(true, nullptr);
+	MenuBuilder.BeginSection("CalibrationComponents", LOCTEXT("CalibrationComponents", "Calibration Point Components"));
+
+	if (AActor* Calibrator = WeakCalibrator.Get())
+	{
+		constexpr uint32 NumInlineAllocations = 32;
+		TArray<UCalibrationPointComponent*, TInlineAllocator<NumInlineAllocations>> CalibrationPointComponents;
+		Calibrator->GetComponents(CalibrationPointComponents);
+
+		for (UCalibrationPointComponent* CalibratorComponent : CalibrationPointComponents)
+		{
+			MenuBuilder.AddMenuEntry(
+				FText::Format(LOCTEXT("ComponentLabel", "{0}"), FText::FromString(CalibratorComponent->GetName())),
+				FText::Format(LOCTEXT("ComponentTooltip", "{0}"), FText::FromString(CalibratorComponent->GetName())),
+				FSlateIcon(),
+				FUIAction(
+					FExecuteAction::CreateUObject(this, &UCameraLensDistortionAlgoAruco::OnCalibrationComponentSelected, CalibratorComponent),
+					FCanExecuteAction(),
+					FIsActionChecked::CreateUObject(this, &UCameraLensDistortionAlgoAruco::IsCalibrationComponentSelected, CalibratorComponent)
+				),
+				NAME_None,
+				EUserInterfaceActionType::ToggleButton
+			);
+		}
+	}
+	else
+	{
+		MenuBuilder.AddMenuEntry(
+			FText::FromName(NAME_None),
+			FText::FromName(NAME_None),
+			FSlateIcon(),
+			FUIAction(
+				FExecuteAction(),
+				FCanExecuteAction(),
+				FIsActionChecked()
+			),
+			NAME_None,
+			EUserInterfaceActionType::None
+		);
+	}
+	MenuBuilder.EndSection();
+
+	return MenuBuilder.MakeWidget();
+}
+
+TSharedRef<SWidget> UCameraLensDistortionAlgoAruco::BuildFocalLengthEstimateWidget()
+{
+	return SNew(SHorizontalBox)
+
+		+ SHorizontalBox::Slot()
+		.FillWidth(0.8)
+		.Padding(0, 0, 10, 0)
+		[
+			SNew(SNumericEntryBox<float>)
+			.Value_UObject(this, &UCameraLensDistortionAlgoAruco::GetFocalLengthEstimate)
+			.OnValueChanged_UObject(this, &UCameraLensDistortionAlgoAruco::SetFocalLengthEstimate)
+		]
+
+		+ SHorizontalBox::Slot()
+		.FillWidth(0.2)
+		[
+			SNew(SCheckBox)
+			.Padding(FMargin(5, 0, 15, 0))
+			.IsChecked_UObject(this, &UCameraLensDistortionAlgoAruco::IsFixFocalLengthChecked)
+			.OnCheckStateChanged_UObject(this, &UCameraLensDistortionAlgoAruco::OnFixFocalLengthCheckStateChanged)
+			[
+				SNew(STextBlock).Text(LOCTEXT("FixText", "Fix?"))
+			]
+		];
+}
+
+TSharedRef<SWidget> UCameraLensDistortionAlgoAruco::BuildFixImageCenterWidget()
+{
+	return SNew(SCheckBox)
+		.IsChecked_UObject(this, &UCameraLensDistortionAlgoAruco::IsFixImageCenterChecked)
+		.OnCheckStateChanged_UObject(this, &UCameraLensDistortionAlgoAruco::OnFixImageCenterCheckStateChanged);
+}
+
+TSharedRef<SWidget> UCameraLensDistortionAlgoAruco::BuildCalibrationActionButtons()
+{
+	return SNew(SHorizontalBox)
+
+		+ SHorizontalBox::Slot() // Button to clear all rows
+		.AutoWidth()
+		[
+			SNew(SButton)
+			.Text(LOCTEXT("ClearAll", "Clear All"))
+			.HAlign(HAlign_Center)
+			.VAlign(VAlign_Center)
+			.OnClicked_UObject(this, &UCameraLensDistortionAlgoAruco::OnClearCalibrationRowsClicked)
+		];
+}
+
+TSharedRef<SWidget> UCameraLensDistortionAlgoAruco::BuildCalibrationRowListWidget()
+{
+	CalibrationRowListWidget = SNew(SListView<TSharedPtr<FLensDistortionArucoRowData>>)
+		.ItemHeight(24)
+		.ListItemsSource(&CalibrationRows)
+		.SelectionMode(ESelectionMode::Multi)
+		.OnGenerateRow_Lambda([](TSharedPtr<FLensDistortionArucoRowData> InItem, const TSharedRef<STableViewBase>& OwnerTable) -> TSharedRef<ITableRow>
+		{
+			return SNew(CameraLensDistortionAlgoAruco::SCalibrationRowGenerator, OwnerTable).CalibrationRowData(InItem);
+		})
+		.OnMouseButtonDoubleClick_Lambda([](TSharedPtr<FLensDistortionArucoRowData> Row)
+		{
+			// Open a dialog showing the media image of the selected row with debug aruco markers drawn on it
+			UTexture2D* DebugImage = UE::CameraCalibration::Private::CreateDebugTexture(Row);
+			UE::CameraCalibration::Private::DisplayDebugImageDialog(DebugImage);
+		})
+		.OnKeyDownHandler_UObject(this, &UCameraLensDistortionAlgoAruco::OnCalibrationRowListKeyPressed)
+		.HeaderRow
+		(
+			SNew(SHeaderRow)
+
+			+ SHeaderRow::Column("Index")
+			.DefaultLabel(LOCTEXT("Index", "Index"))
+			.FillWidth(0.1f)
+
+			+ SHeaderRow::Column("NumMarkers")
+			.DefaultLabel(LOCTEXT("NumMarkers", "Num Markers"))
+			.FillWidth(0.2f)
+
+			+ SHeaderRow::Column("Image")
+			.DefaultLabel(LOCTEXT("Image", "Image"))
+			.FillWidth(0.7f)
+		);
+
+	return CalibrationRowListWidget.ToSharedRef();
+}
+
+void UCameraLensDistortionAlgoAruco::OnCalibratorSelected(const FAssetData& AssetData)
+{
+	if (AssetData.IsValid())
+	{
+		SetCalibrator(Cast<AActor>(AssetData.GetAsset()));
+	}
+}
+
+bool UCameraLensDistortionAlgoAruco::DoesAssetHaveCalibrationComponent(const FAssetData& AssetData) const
+{
+	if (const AActor* Actor = Cast<AActor>(AssetData.GetAsset()))
+	{
+		constexpr uint32 NumInlineAllocations = 32;
+		TArray<UCalibrationPointComponent*, TInlineAllocator<NumInlineAllocations>> CalibrationPoints;
+		Actor->GetComponents(CalibrationPoints);
+
+		return (CalibrationPoints.Num() > 0);
+	}
+
+	return false;
+}
+
+FAssetData UCameraLensDistortionAlgoAruco::GetCalibratorAssetData() const
+{
+	return FAssetData(WeakCalibrator.Get(), true);
+}
+
+FText UCameraLensDistortionAlgoAruco::GetCalibrationComponentMenuText() const 
+{
+	if (ActiveCalibratorComponents.Num() > 1)
+	{
+		return LOCTEXT("MultipleCalibrationComponents", "Multiple Values");
+	}
+	else if (ActiveCalibratorComponents.Num() == 1)
+	{
+		if (UCalibrationPointComponent* CalibratorComponent = ActiveCalibratorComponents[0].Get())
+		{
+			return FText::FromString(CalibratorComponent->GetName());
+		}
+	}
+
+	return LOCTEXT("NoCalibrationComponents", "None");
+}
+
+void UCameraLensDistortionAlgoAruco::OnCalibrationComponentSelected(UCalibrationPointComponent* SelectedComponent)
+{
+	if (IsCalibrationComponentSelected(SelectedComponent))
+	{
+		ActiveCalibratorComponents.Remove(SelectedComponent);
+	}
+	else
+	{
+		ActiveCalibratorComponents.Add(SelectedComponent);
+	}
+}
+
+bool UCameraLensDistortionAlgoAruco::IsCalibrationComponentSelected(UCalibrationPointComponent* SelectedComponent) const
+{
+	return ActiveCalibratorComponents.Contains(SelectedComponent);
+}
+
+TOptional<float> UCameraLensDistortionAlgoAruco::GetFocalLengthEstimate() const
+{
+	return FocalLengthEstimate;
+}
+
+void UCameraLensDistortionAlgoAruco::SetFocalLengthEstimate(float NewValue)
+{
+	FocalLengthEstimate = NewValue;
+}
+
+ECheckBoxState UCameraLensDistortionAlgoAruco::IsFixFocalLengthChecked() const
+{
+	return bFixFocalLength ? ECheckBoxState::Checked : ECheckBoxState::Unchecked;
+}
+
+void UCameraLensDistortionAlgoAruco::OnFixFocalLengthCheckStateChanged(ECheckBoxState NewState)
+{
+	bFixFocalLength = (NewState == ECheckBoxState::Checked);
+}
+
+ECheckBoxState UCameraLensDistortionAlgoAruco::IsFixImageCenterChecked() const
+{
+	return bFixImageCenter ? ECheckBoxState::Checked : ECheckBoxState::Unchecked;
+}
+
+void UCameraLensDistortionAlgoAruco::OnFixImageCenterCheckStateChanged(ECheckBoxState NewState)
+{
+	bFixImageCenter = (NewState == ECheckBoxState::Checked);
+}
+
+FReply UCameraLensDistortionAlgoAruco::OnClearCalibrationRowsClicked()
+{
+	ClearCalibrationRows();
+	return FReply::Handled();
+}
+
+FReply UCameraLensDistortionAlgoAruco::OnCalibrationRowListKeyPressed(const FGeometry& Geometry, const FKeyEvent& KeyEvent)
+{
+	if (!CalibrationRowListWidget)
+	{
+		return FReply::Unhandled();
+	}
+
+	if (KeyEvent.GetKey() == EKeys::Delete)
+	{
+		// Delete selected items
+		const TArray<TSharedPtr<FLensDistortionArucoRowData>> SelectedItems = CalibrationRowListWidget->GetSelectedItems();
+
+		for (const TSharedPtr<FLensDistortionArucoRowData>& SelectedItem : SelectedItems)
+		{
+			CalibrationRows.Remove(SelectedItem);
+
+			// Delete the previously exported .json file for the row that is being deleted
+			if (ULensDistortionTool* LensDistortionTool = WeakTool.Get())
+			{
+				LensDistortionTool->DeleteExportedRow(SelectedItem->Index);
+			}
+		}
+
+		CalibrationRowListWidget->RequestListRefresh();
+		return FReply::Handled();
+	}
+	else if (KeyEvent.GetModifierKeys().IsControlDown() && (KeyEvent.GetKey() == EKeys::A))
+	{
+		// Select all items
+		CalibrationRowListWidget->SetItemSelection(CalibrationRows, true);
+		return FReply::Handled();
+	}
+	else if (KeyEvent.GetKey() == EKeys::Escape)
+	{
+		// De-select all items
+		CalibrationRowListWidget->ClearSelection();
+		return FReply::Handled();
+	}
+
+	return FReply::Unhandled();
+}
+
+#undef LOCTEXT_NAMESPACE
diff --git a/Engine/Plugins/VirtualProduction/CameraCalibration/Source/CameraCalibrationEditor/Private/Calibrators/CameraLensDistortionAlgoAruco.h b/Engine/Plugins/VirtualProduction/CameraCalibration/Source/CameraCalibrationEditor/Private/Calibrators/CameraLensDistortionAlgoAruco.h
new file mode 100644
--- /dev/null
+++ b/Engine/Plugins/VirtualProduction/CameraCalibration/Source/CameraCalibrationEditor/Private/Calibrators/CameraLensDistortionAlgoAruco.h
@@ -0,0 +1,182 @@
+// Copyright Epic Games, Inc. All Rights Reserved.
+
+#pragma once
+
+#include "CameraLensDistortionAlgo.h"
+#include "ImageCore.h"
+
+#include "CameraLensDistortionAlgoAruco.generated.h"
+
+class UCalibrationPointComponent;
+
+template <typename ItemType>
+class SListView;
+
+struct FArucoCalibrationPoint;
+
+/** Holds information of the calibration row */
+USTRUCT()
+struct FLensDistortionArucoRowData
+{
+	GENERATED_BODY()
+
+	/** Index to display in list */
+	UPROPERTY()
+	int32 Index = -1;
+
+	/** Array of aruco points (with their 3D and 2D corner locations) */
+	UPROPERTY()
+	TArray<FArucoCalibrationPoint> ArucoPoints;
+
+	/** Pose of the tracked camera when the aruco points were detected */
+	UPROPERTY()
+	FTransform CameraPose = FTransform::Identity;
+
+	/** Cached image of the media texture that contained the aruco markers */
+	FImage MediaImage;
+};
+
+/** 
+ * Implements a lens distortion calibration algorithm. It requires a Aruco pattern
+ */
+UCLASS()
+class UCameraLensDistortionAlgoAruco : public UCameraLensDistortionAlgo
+{
+	GENERATED_BODY()
+	
+public:
+
+	//~ Begin CameraLensDistortionAlgo
+	virtual void Initialize(ULensDistortionTool* InTool) override;
+	virtual void Shutdown() override;
+	virtual bool SupportsModel(const TSubclassOf<ULensModel>& LensModel) const override;
+	virtual void Tick(float DeltaTime) override;
+	virtual bool OnViewportClicked(const FGeometry& MyGeometry, const FPointerEvent& MouseEvent) override;
+	virtual TSharedRef<SWidget> BuildUI() override;
+	virtual FName FriendlyName() const override { return TEXT("Lens Distortion Aruco"); };
+	virtual FName ShortName() const override { return TEXT("Aruco"); };
+	virtual void OnDistortionSavedToLens() override;
+	virtual bool GetLensDistortion(
+		float& OutFocus,
+		float& OutZoom,
+		FDistortionInfo& OutDistortionInfo,
+		FFocalLengthInfo& OutFocalLengthInfo,
+		FImageCenterInfo& OutImageCenterInfo,
+		TSubclassOf<ULensModel>& OutLensModel,
+		double& OutError,
+		FText& OutErrorMessage
+	) override;
+	virtual bool HasCalibrationData() const override;
+	virtual void PreImportCalibrationData() override;
+	virtual int32 ImportCalibrationRow(const TSharedRef<FJsonObject>& CalibrationRowObject, const FImage& RowImage) override;
+	virtual void PostImportCalibrationData() override;
+	//~ End CameraLensDistortionAlgo
+
+private:
+	/** Add a new calibration row from media texture and camera data */
+	bool AddCalibrationRow(FText& OutErrorMessage);
+
+	/** Clears the list of calibration sample rows */
+	void ClearCalibrationRows();
+
+	/** Sets the calibrator object to be used. Updates the selection picker. */
+	void SetCalibrator(AActor* InCalibrator);
+
+	/** Export global session data to a .json file */
+	void ExportSessionData();
+
+	/** Export the row data to a json file */
+	void ExportRow(TSharedPtr<FLensDistortionArucoRowData> Row);
+
+private:
+	/** Builds the UI of the calibration device picker */
+	TSharedRef<SWidget> BuildCalibratorPickerWidget();
+
+	/** Builds the UI of the fix focal length width */
+	TSharedRef<SWidget> BuildFocalLengthEstimateWidget();
+
+	/** Builds the UI of the fix image center widget */
+	TSharedRef<SWidget> BuildFixImageCenterWidget();
+
+	/** Builds the UI of the calibration points table */
+	TSharedRef<SWidget> BuildCalibrationRowListWidget();
+
+	/** Builds the UI for the action buttons (RemoveLast, ClearAll) */
+	TSharedRef<SWidget> BuildCalibrationActionButtons();
+
+	/** Builds the UI of the calibration component picker */
+	TSharedRef<SWidget> BuildCalibrationComponentPickerWidget();
+
+	/** Builds the UI of the calibration component picker */
+	TSharedRef<SWidget> BuildCalibrationComponentMenu();
+
+	/** Gets the actor from the input asset data and caches it as the current calibrator */
+	void OnCalibratorSelected(const FAssetData& AssetData);
+
+	/** Returns true if the actor represented by the input asset data has any calibration point components. Used to filter calibrator actors in the picker. */
+	bool DoesAssetHaveCalibrationComponent(const FAssetData& AssetData) const;
+
+	/** Get an FAssetData representation of the current calibrator actor */
+	FAssetData GetCalibratorAssetData() const;
+
+	/** Get the text to display in the calibration component combobox. */
+	FText GetCalibrationComponentMenuText() const;
+
+	/** Add or remove the selected component from the set of active components and refresh the calibrator point combobox */
+	void OnCalibrationComponentSelected(UCalibrationPointComponent* SelectedComponent);
+
+	/** Returns true if the input component is in the set of active calibration components */
+	bool IsCalibrationComponentSelected(UCalibrationPointComponent* SelectedComponent) const;
+
+	/** Get the current value of the focal length estimate */
+	TOptional<float> GetFocalLengthEstimate() const;
+
+	/** Set the current value of the focal length estimate */
+	void SetFocalLengthEstimate(float NewValue);
+
+	/** Returns the state of the fix focal length checkbox */
+	ECheckBoxState IsFixFocalLengthChecked() const;
+
+	/** Changes the state of the fix focal length checkbox */
+	void OnFixFocalLengthCheckStateChanged(ECheckBoxState NewState);
+
+	/** Returns the state of the fix image center checkbox */
+	ECheckBoxState IsFixImageCenterChecked() const;
+
+	/** Changes the state of the fix image center checkbox */
+	void OnFixImageCenterCheckStateChanged(ECheckBoxState NewState);
+
+	/** Event triggered when the Clear button is pressed. Clears all of the captured calibration rows */
+	FReply OnClearCalibrationRowsClicked();
+
+	/** Event triggered when a key is pressed on the calibration list view */
+	FReply OnCalibrationRowListKeyPressed(const FGeometry& Geometry, const FKeyEvent& KeyEvent);
+
+private:
+	/** Version of the current calibration dataset */
+	static const int DATASET_VERSION;
+
+	/** The lens distortion tool controller */
+	TWeakObjectPtr<ULensDistortionTool> WeakTool;
+
+	/** The currently selected calibrator object. It is expected to contain one or more UCalibrationPointComponent in it */
+	TWeakObjectPtr<AActor> WeakCalibrator;
+
+	/** Rows source for the CalibrationRowListWidget */
+	TArray<TSharedPtr<FLensDistortionArucoRowData>> CalibrationRows;
+
+	/** Displays the list of calibration points that will be used to calculate the lens distortion */
+	TSharedPtr<SListView<TSharedPtr<FLensDistortionArucoRowData>>> CalibrationRowListWidget;
+
+	/** If true, the solver will not solve for focal length, and will just use the input value */
+	bool bFixFocalLength = false;
+
+	/** If true, the solver will not solve for image center, and will just use the input value */
+	bool bFixImageCenter = false;
+
+	/** Estimate for the focal length to provide to the solver. If not set, user will be warned to set before continuing calibration */
+	TOptional<float> FocalLengthEstimate;
+
+	/** Container for the set of calibrator components selected in the component combobox */
+	TArray<TWeakObjectPtr<UCalibrationPointComponent>> ActiveCalibratorComponents;
+};
diff --git a/Engine/Plugins/VirtualProduction/CameraCalibration/Source/CameraCalibrationEditor/Private/Calibrators/CameraLensDistortionAlgoCheckerboard.cpp b/Engine/Plugins/VirtualProduction/CameraCalibration/Source/CameraCalibrationEditor/Private/Calibrators/CameraLensDistortionAlgoCheckerboard.cpp
--- a/Engine/Plugins/VirtualProduction/CameraCalibration/Source/CameraCalibrationEditor/Private/Calibrators/CameraLensDistortionAlgoCheckerboard.cpp
+++ b/Engine/Plugins/VirtualProduction/CameraCalibration/Source/CameraCalibrationEditor/Private/Calibrators/CameraLensDistortionAlgoCheckerboard.cpp
@@ -19,8 +19,8 @@
 #include "JsonObjectConverter.h"
 #include "LensComponent.h"
 #include "LensFile.h"
-#include "Misc/MessageDialog.h"
 #include "Materials/MaterialInstanceDynamic.h"
+#include "Misc/MessageDialog.h"
 #include "Models/AnamorphicLensModel.h"
 #include "Models/SphericalLensModel.h"
 #include "OpenCVHelper.h"
@@ -143,7 +143,7 @@ void UCameraLensDistortionAlgoCheckerboard::Initialize(ULensDistortionTool* InTo
 	// Guess which calibrator to use.
 	SetCalibrator(FindFirstCalibrator());
 
-	// Initialize coverage matrix
+	// Initialize coverage texture
 	FCameraCalibrationStepsController* StepsController = Tool->GetCameraCalibrationStepsController();
 
 	if (!ensure(StepsController))
@@ -310,16 +310,9 @@ bool UCameraLensDistortionAlgoCheckerboard::AddCalibrationRow(FText& OutErrorMes
 
 	TArray<FColor> Pixels;
 	FIntPoint Size;
-	ETextureRenderTargetFormat PixelFormat;
-
-	if (!StepsController->ReadMediaPixels(Pixels, Size, PixelFormat, OutErrorMessage, ESimulcamViewportPortion::CameraFeed))
-	{
-		return false;
-	}
 
-	if (PixelFormat != ETextureRenderTargetFormat::RTF_RGBA8)
+	if (!StepsController->ReadMediaPixels(Pixels, Size, OutErrorMessage, ESimulcamViewportPortion::CameraFeed))
 	{
-		OutErrorMessage = LOCTEXT("InvalidFormat", "MediaPlateRenderTarget did not have the expected RTF_RGBA8 format");
 		return false;
 	}
 
diff --git a/Engine/Plugins/VirtualProduction/CameraCalibration/Source/CameraCalibrationEditor/Private/Calibrators/CameraNodalOffsetAlgoAruco.cpp b/Engine/Plugins/VirtualProduction/CameraCalibration/Source/CameraCalibrationEditor/Private/Calibrators/CameraNodalOffsetAlgoAruco.cpp
--- a/Engine/Plugins/VirtualProduction/CameraCalibration/Source/CameraCalibrationEditor/Private/Calibrators/CameraNodalOffsetAlgoAruco.cpp
+++ b/Engine/Plugins/VirtualProduction/CameraCalibration/Source/CameraCalibrationEditor/Private/Calibrators/CameraNodalOffsetAlgoAruco.cpp
@@ -119,19 +119,12 @@ bool UCameraNodalOffsetAlgoAruco::PopulatePoints(FText& OutErrorMessage)
 
 	TArray<FColor> Pixels;
 	FIntPoint Size;
-	ETextureRenderTargetFormat PixelFormat;
 
-	if (!StepsController->ReadMediaPixels(Pixels, Size, PixelFormat, OutErrorMessage, ESimulcamViewportPortion::CameraFeed))
+	if (!StepsController->ReadMediaPixels(Pixels, Size, OutErrorMessage, ESimulcamViewportPortion::CameraFeed))
 	{
 		return false;
 	}
 
-	if (PixelFormat != ETextureRenderTargetFormat::RTF_RGBA8)
-	{
-		OutErrorMessage = LOCTEXT("InvalidFormat", "MediaPlateRenderTarget did not have the expected RTF_RGBA8 format");
-		return false;
-	}
-
 	TArray<FArucoMarker> IdentifiedMarkers;
 	bool bResult = FOpenCVHelper::IdentifyArucoMarkers(Pixels, Size, Dictionary, IdentifiedMarkers);
 
diff --git a/Engine/Plugins/VirtualProduction/CameraCalibration/Source/CameraCalibrationEditor/Private/Calibrators/CameraNodalOffsetAlgoCheckerboard.cpp b/Engine/Plugins/VirtualProduction/CameraCalibration/Source/CameraCalibrationEditor/Private/Calibrators/CameraNodalOffsetAlgoCheckerboard.cpp
--- a/Engine/Plugins/VirtualProduction/CameraCalibration/Source/CameraCalibrationEditor/Private/Calibrators/CameraNodalOffsetAlgoCheckerboard.cpp
+++ b/Engine/Plugins/VirtualProduction/CameraCalibration/Source/CameraCalibrationEditor/Private/Calibrators/CameraNodalOffsetAlgoCheckerboard.cpp
@@ -111,19 +111,12 @@ bool UCameraNodalOffsetAlgoCheckerboard::PopulatePoints(FText& OutErrorMessage)
 
 	TArray<FColor> Pixels;
 	FIntPoint Size;
-	ETextureRenderTargetFormat PixelFormat;
 
-	if (!StepsController->ReadMediaPixels(Pixels, Size, PixelFormat, OutErrorMessage, ESimulcamViewportPortion::CameraFeed))
+	if (!StepsController->ReadMediaPixels(Pixels, Size, OutErrorMessage, ESimulcamViewportPortion::CameraFeed))
 	{
 		return false;
 	}
 
-	if (PixelFormat != ETextureRenderTargetFormat::RTF_RGBA8)
-	{
-		OutErrorMessage = LOCTEXT("InvalidFormat", "MediaPlateRenderTarget did not have the expected RTF_RGBA8 format");
-		return false;
-	}
-
 	const ACameraCalibrationCheckerboard* Checkerboard = Cast<ACameraCalibrationCheckerboard>(Calibrator.Get());
 	check(Checkerboard); // This should be ensured by the picker.
 
diff --git a/Engine/Plugins/VirtualProduction/CameraCalibration/Source/CameraCalibrationEditor/Private/CameraCalibrationUtilsPrivate.cpp b/Engine/Plugins/VirtualProduction/CameraCalibration/Source/CameraCalibrationEditor/Private/CameraCalibrationUtilsPrivate.cpp
--- a/Engine/Plugins/VirtualProduction/CameraCalibration/Source/CameraCalibrationEditor/Private/CameraCalibrationUtilsPrivate.cpp
+++ b/Engine/Plugins/VirtualProduction/CameraCalibration/Source/CameraCalibrationEditor/Private/CameraCalibrationUtilsPrivate.cpp
@@ -3,8 +3,11 @@
 #include "CameraCalibrationUtilsPrivate.h"
 
 #include "CalibrationPointComponent.h"
+#include "Editor.h"
+#include "Engine/Level.h"
 #include "Engine/Texture2D.h"
 #include "TextureResource.h"
+#include "UObject/UObjectIterator.h"
 
 namespace UE::CameraCalibration::Private
 {
@@ -102,26 +105,118 @@ namespace UE::CameraCalibration::Private
 		return Dictionary;
 	}
 
+	void FindActorsWithCalibrationComponents(TArray<AActor*>& ActorsWithCalibrationComponents)
+	{
+		ActorsWithCalibrationComponents.Empty();
+
+		const UWorld* World = GEditor ? GEditor->GetEditorWorldContext().World() : nullptr;
+		const EObjectFlags ExcludeFlags = RF_ClassDefaultObject; // We don't want the calibrator CDOs.
+
+		for (TObjectIterator<UCalibrationPointComponent> It(ExcludeFlags, true, EInternalObjectFlags::Garbage); It; ++It)
+		{
+			AActor* Actor = It->GetOwner();
+			
+			if (Actor && (Actor->GetWorld() == World))
+			{
+				// Exclude any actors belonging to a level that is not currently visible
+				ULevel* Level = Actor->GetLevel();
+				if (Level && Level->bIsVisible)
+				{
+					ActorsWithCalibrationComponents.Add(Actor);
+				}
+			}
+		}
+	}
+
+	void SortArucoCalibrationPoints(const TArray<FArucoCalibrationPoint>& ArucoCalibrationPoints, TMap<UObject*, TArray<FArucoCalibrationPoint>>& OutSortedArucoCalibrationPointMap)
+	{
+		for (const FArucoCalibrationPoint& ArucoCalibrationPoint : ArucoCalibrationPoints)
+		{
+			if (TArray<FArucoCalibrationPoint>* ArucoCalibrationPointSet = OutSortedArucoCalibrationPointMap.Find(ArucoCalibrationPoint.Owner))
+			{
+				ArucoCalibrationPointSet->Add(ArucoCalibrationPoint);
+			}
+			else
+			{
+				TArray<FArucoCalibrationPoint> NewArucoCalibrationPointSet;
+				NewArucoCalibrationPointSet.Add(ArucoCalibrationPoint);
+
+				OutSortedArucoCalibrationPointMap.Add(ArucoCalibrationPoint.Owner, NewArucoCalibrationPointSet);
+			}
+		}
+	}
+
+	bool FindArucoCalibrationPoint(const TArray<UCalibrationPointComponent*>& CalibrationComponents, EArucoDictionary ArucoDictionary, const FArucoMarker& ArucoMarker, FArucoCalibrationPoint& OutArucoCalibrationPoint)
+	{
+		// Build calibrator point name based on the detected marker
+		const FString DictionaryName = GetArucoDictionaryName(ArucoDictionary);
+
+		constexpr int32 NumExpectedCorners = 4;
+		static const TArray<FString> CornerNames = { TEXT("TL"), TEXT("TR"), TEXT("BR"), TEXT("BL") };
+
+		OutArucoCalibrationPoint.MarkerID = ArucoMarker.MarkerID;
+		OutArucoCalibrationPoint.Name = FString::Printf(TEXT("%s-%d"), *DictionaryName, ArucoMarker.MarkerID);
+
+		for (UCalibrationPointComponent* Component : CalibrationComponents)
+		{
+			if (!Component)
+			{
+				continue;
+			}
+
+			// Check each corner to see if it belongs to the current calibration component (lookup is based on the corner name)
+			int32 FoundCorners = 0;
+			for (int32 CornerIndex = 0; CornerIndex < NumExpectedCorners; ++CornerIndex)
+			{
+				const FString CornerName = FString::Printf(TEXT("%s-%s"), *OutArucoCalibrationPoint.Name, *CornerNames[CornerIndex]); //-V557
+				OutArucoCalibrationPoint.Corners2D[CornerIndex] = ArucoMarker.Corners[CornerIndex];
+
+				// If the current corner is a subpoint of the current calibration component being inspected, get its 3D world location
+				FVector Corner;
+				if (Component->GetWorldLocation(CornerName, Corner))
+				{
+					OutArucoCalibrationPoint.Corners3D[CornerIndex] = Corner;
+					++FoundCorners;
+				}
+			}
+
+			// If all four corners were found, we're done. Otherwise, check the next calibration component for the four corners
+			if (FoundCorners == NumExpectedCorners)
+			{
+				OutArucoCalibrationPoint.Owner = Component;
+				return true;
+			}
+		}
+
+		return false;
+	}
+
 	void ClearTexture(UTexture2D* Texture, FColor ClearColor)
 	{
-		TArray<FColor> Pixels;
-		Pixels.Init(ClearColor, Texture->GetSizeX() * Texture->GetSizeY());
+		if (Texture)
+		{
+			TArray<FColor> Pixels;
+			Pixels.Init(ClearColor, Texture->GetSizeX() * Texture->GetSizeY());
 
-		void* TextureData = Texture->GetPlatformData()->Mips[0].BulkData.Lock(LOCK_READ_WRITE);
+			void* TextureData = Texture->GetPlatformData()->Mips[0].BulkData.Lock(LOCK_READ_WRITE);
 
-		FMemory::Memcpy(TextureData, Pixels.GetData(), Pixels.Num() * sizeof(FColor));
+			FMemory::Memcpy(TextureData, Pixels.GetData(), Pixels.Num() * sizeof(FColor));
 
-		Texture->GetPlatformData()->Mips[0].BulkData.Unlock();
-		Texture->UpdateResource();
+			Texture->GetPlatformData()->Mips[0].BulkData.Unlock();
+			Texture->UpdateResource();
+		}
 	}
 
 	void SetTextureData(UTexture2D* Texture, const TArray<FColor>& PixelData)
 	{
-		void* TextureData = Texture->GetPlatformData()->Mips[0].BulkData.Lock(LOCK_READ_WRITE);
+		if (Texture)
+		{
+			void* TextureData = Texture->GetPlatformData()->Mips[0].BulkData.Lock(LOCK_READ_WRITE);
 
-		FMemory::Memcpy(TextureData, PixelData.GetData(), PixelData.Num() * sizeof(FColor));
+			FMemory::Memcpy(TextureData, PixelData.GetData(), PixelData.Num() * sizeof(FColor));
 
-		Texture->GetPlatformData()->Mips[0].BulkData.Unlock();
-		Texture->UpdateResource();
+			Texture->GetPlatformData()->Mips[0].BulkData.Unlock();
+			Texture->UpdateResource();
+		}
 	}
-}
\ No newline at end of file
+}
diff --git a/Engine/Plugins/VirtualProduction/CameraCalibration/Source/CameraCalibrationEditor/Private/CameraCalibrationUtilsPrivate.h b/Engine/Plugins/VirtualProduction/CameraCalibration/Source/CameraCalibrationEditor/Private/CameraCalibrationUtilsPrivate.h
--- a/Engine/Plugins/VirtualProduction/CameraCalibration/Source/CameraCalibrationEditor/Private/CameraCalibrationUtilsPrivate.h
+++ b/Engine/Plugins/VirtualProduction/CameraCalibration/Source/CameraCalibrationEditor/Private/CameraCalibrationUtilsPrivate.h
@@ -3,6 +3,39 @@
 
 #include "OpenCVHelper.h"
 
+#include "CameraCalibrationUtilsPrivate.generated.h"
+
+class UCalibrationPointComponent;
+
+/** Structure representing an aruco marker, including ID, Name, and corner coordinates in 2D image space and 3D world space */
+USTRUCT()
+struct FArucoCalibrationPoint
+{
+	GENERATED_BODY()
+
+	/** 3D locations (in world space) of each corner */
+	UPROPERTY()
+	FVector Corners3D[4];
+
+	/** 2D locations (in pixels) of each corner */
+	UPROPERTY()
+	FVector2f Corners2D[4];
+
+	/** Aruco Marker ID */
+	UPROPERTY()
+	int32 MarkerID = -1;
+	
+	/** Name of the marker (follows the naming convention of "[DictionaryName]-[MarkerID]") */
+	UPROPERTY()
+	FString Name;
+
+	/** 
+	 * Object (almost certainly a CalibrationPointComponent) that owns this aruco point. 
+	 * Used to sort detected points that belong to different patterns in the same image 
+	 */
+	TObjectPtr<UObject> Owner = nullptr;
+};
+
 namespace UE::CameraCalibration::Private
 {
 	/** Returns the aruco dictionary that matches the calibration point component(s) of the input calibrator actor */
@@ -14,6 +47,15 @@ namespace UE::CameraCalibration::Private
 	/** Get the string representation of the input aruco dictionary */
 	FString GetArucoDictionaryName(EArucoDictionary Dictionary);
 
+	/** Find all actors in the current world that have a calibration point component attached (excluding CDOs and actors in levels that are not currently visible) */
+	void FindActorsWithCalibrationComponents(TArray<AActor*>& ActorsWithCalibrationComponents);
+
+	/** Find an aruco marker calibration point in one of the input calibration components that matches the input dictionary and marker ID */
+	bool FindArucoCalibrationPoint(const TArray<UCalibrationPointComponent*>& CalibrationComponents, EArucoDictionary ArucoDictionary, const FArucoMarker& ArucoMarker, FArucoCalibrationPoint& OutArucoCalibrationPoint);
+
+	/** Sort input calibration points based on owner */
+	void SortArucoCalibrationPoints(const TArray<FArucoCalibrationPoint>& ArucoCalibrationPoints, TMap<UObject*, TArray<FArucoCalibrationPoint>>& OutSortedArucoCalibrationPointMap);
+
 	/** Set every pixel in the input texture to the clear color */
 	void ClearTexture(UTexture2D* Texture, FColor ClearColor = FColor::Transparent);
 
diff --git a/Engine/Plugins/VirtualProduction/CameraCalibration/Source/CameraCalibrationEditor/Private/UI/SImageTexture.cpp b/Engine/Plugins/VirtualProduction/CameraCalibration/Source/CameraCalibrationEditor/Private/UI/SImageTexture.cpp
new file mode 100644
--- /dev/null
+++ b/Engine/Plugins/VirtualProduction/CameraCalibration/Source/CameraCalibrationEditor/Private/UI/SImageTexture.cpp
@@ -0,0 +1,44 @@
+// Copyright Epic Games, Inc. All Rights Reserved.
+
+#include "SImageTexture.h"
+
+#include "Widgets/Layout/SBox.h"
+#include "Widgets/Images/SImage.h"
+
+void SImageTexture::Construct(const FArguments& InArgs, UTexture2D* InTexture)
+{
+	if (!InTexture)
+	{
+		ChildSlot
+		[
+			SNullWidget::NullWidget
+		];
+
+		return;
+	}
+
+	Texture = TStrongObjectPtr<UTexture2D>(InTexture);
+
+	Brush.SetResourceObject(InTexture);
+	Brush.ImageSize = FVector2f(Texture->GetSurfaceWidth(), Texture->GetSurfaceHeight());
+
+	float AspectRatio = 1.0f;
+	if (!FMath::IsNearlyZero(Texture->GetSurfaceHeight()))
+	{
+		AspectRatio = Texture->GetSurfaceWidth() / Texture->GetSurfaceHeight();
+	}
+
+	ChildSlot
+	[
+		SNew(SBox)
+		.MinAspectRatio(AspectRatio)
+		.MaxAspectRatio(AspectRatio)
+		.MinDesiredWidth(InArgs._MinDesiredWidth)
+		.MinDesiredHeight(InArgs._MinDesiredHeight)
+		.MaxDesiredWidth(InArgs._MaxDesiredWidth)
+		.MaxDesiredHeight(InArgs._MaxDesiredHeight)
+		[
+			SNew(SImage).Image(&Brush)
+		]
+	];
+}
diff --git a/Engine/Plugins/VirtualProduction/CameraCalibration/Source/CameraCalibrationEditor/Private/UI/SImageTexture.h b/Engine/Plugins/VirtualProduction/CameraCalibration/Source/CameraCalibrationEditor/Private/UI/SImageTexture.h
new file mode 100644
--- /dev/null
+++ b/Engine/Plugins/VirtualProduction/CameraCalibration/Source/CameraCalibrationEditor/Private/UI/SImageTexture.h
@@ -0,0 +1,38 @@
+// Copyright Epic Games, Inc. All Rights Reserved.
+#pragma once
+
+#include "Widgets/SCompoundWidget.h"
+
+#include "CoreMinimal.h"
+#include "Engine/Texture2D.h"
+#include "Styling/SlateBrush.h"
+#include "Types/SlateStructs.h"
+#include "UObject/GCObject.h"
+#include "UObject/StrongObjectPtr.h"
+#include "Widgets/DeclarativeSyntaxSupport.h"
+
+
+/** Widget that holds a reference to a UTexture and creates an SImage for that Texture using a brush */
+class SImageTexture : public SCompoundWidget
+{
+	SLATE_BEGIN_ARGS(SImageTexture)
+		: _MinDesiredWidth(FOptionalSize())
+		, _MinDesiredHeight(FOptionalSize())
+		, _MaxDesiredWidth(FOptionalSize())
+		, _MaxDesiredHeight(FOptionalSize())
+	{}
+
+	SLATE_ATTRIBUTE(FOptionalSize, MinDesiredWidth)
+	SLATE_ATTRIBUTE(FOptionalSize, MinDesiredHeight)
+	SLATE_ATTRIBUTE(FOptionalSize, MaxDesiredWidth)
+	SLATE_ATTRIBUTE(FOptionalSize, MaxDesiredHeight)
+
+	SLATE_END_ARGS()
+
+public:
+	void Construct(const FArguments& InArgs, UTexture2D* InTexture);
+
+private:
+	FSlateBrush Brush;
+	TStrongObjectPtr<UTexture2D> Texture;
+};
