CameraCalibration: Fix CIS error, redefinition of cvar

diff --git a/Engine/Plugins/VirtualProduction/CameraCalibration/Source/CameraCalibrationEditor/Private/Calibrators/CameraLensDistortionAlgoAruco.cpp b/Engine/Plugins/VirtualProduction/CameraCalibration/Source/CameraCalibrationEditor/Private/Calibrators/CameraLensDistortionAlgoAruco.cpp
--- a/Engine/Plugins/VirtualProduction/CameraCalibration/Source/CameraCalibrationEditor/Private/Calibrators/CameraLensDistortionAlgoAruco.cpp
+++ b/Engine/Plugins/VirtualProduction/CameraCalibration/Source/CameraCalibrationEditor/Private/Calibrators/CameraLensDistortionAlgoAruco.cpp
@@ -35,10 +35,10 @@
 #define LOCTEXT_NAMESPACE "CameraLensDistortionAlgoAruco"
 
 #if WITH_EDITOR
-static TAutoConsoleVariable<bool> CVarUseIntrinsicsGuess(TEXT("LensDistortionAruco.UseIntrinsicsGuess"), true, TEXT("If true, the solver initializes the camera intrinsics to a user-provided estimate. Otherwise, the solver will compute the initial values."));
-static TAutoConsoleVariable<bool> CVarFixExtrinsics(TEXT("LensDistortionAruco.FixExtrinsics"), false, TEXT("If true, the solver will fix the camera extrinsics to the user-provided camera poses"));
-static TAutoConsoleVariable<bool> CVarFixZeroDistortion(TEXT("LensDistortionAruco.FixZeroDistortion"), false, TEXT("If true, the solver will fix all distortion values to always be 0"));
-static TAutoConsoleVariable<bool> CVarUseExtrinsicsGuess(TEXT("LensDistortionAruco.UseExtrinsicsGuess"), false, TEXT("If true, the actual calibrator and camera poses will be used when running the solver"));
+static TAutoConsoleVariable<bool> CVarUseIntrinsicsGuessAruco(TEXT("LensDistortionAruco.UseIntrinsicsGuess"), true, TEXT("If true, the solver initializes the camera intrinsics to a user-provided estimate. Otherwise, the solver will compute the initial values."));
+static TAutoConsoleVariable<bool> CVarFixExtrinsicsAruco(TEXT("LensDistortionAruco.FixExtrinsics"), false, TEXT("If true, the solver will fix the camera extrinsics to the user-provided camera poses"));
+static TAutoConsoleVariable<bool> CVarFixZeroDistortionAruco(TEXT("LensDistortionAruco.FixZeroDistortion"), false, TEXT("If true, the solver will fix all distortion values to always be 0"));
+static TAutoConsoleVariable<bool> CVarUseExtrinsicsGuessAruco(TEXT("LensDistortionAruco.UseExtrinsicsGuess"), false, TEXT("If true, the actual calibrator and camera poses will be used when running the solver"));
 #endif
 
 const int UCameraLensDistortionAlgoAruco::DATASET_VERSION = 1;
@@ -377,22 +377,22 @@ bool UCameraLensDistortionAlgoAruco::GetLensDistortion(
 
 	ECalibrationFlags SolverFlags = ECalibrationFlags::None;
 
-	if (CVarUseExtrinsicsGuess.GetValueOnGameThread())
+	if (CVarUseExtrinsicsGuessAruco.GetValueOnGameThread())
 	{
 		EnumAddFlags(SolverFlags, ECalibrationFlags::UseExtrinsicGuess);
 	}
 
-	if (CVarUseIntrinsicsGuess.GetValueOnAnyThread())
+	if (CVarUseIntrinsicsGuessAruco.GetValueOnAnyThread())
 	{
 		EnumAddFlags(SolverFlags, ECalibrationFlags::UseIntrinsicGuess);
 	}
 
-	if (CVarFixExtrinsics.GetValueOnAnyThread())
+	if (CVarFixExtrinsicsAruco.GetValueOnAnyThread())
 	{
 		EnumAddFlags(SolverFlags, ECalibrationFlags::FixExtrinsics);
 	}
 
-	if (CVarFixZeroDistortion.GetValueOnAnyThread())
+	if (CVarFixZeroDistortionAruco.GetValueOnAnyThread())
 	{
 		EnumAddFlags(SolverFlags, ECalibrationFlags::FixZeroDistortion);
 	}
@@ -432,7 +432,7 @@ bool UCameraLensDistortionAlgoAruco::GetLensDistortion(
 	OutFocus = LensFileEvalInputs.Focus;
 	OutZoom = LensFileEvalInputs.Zoom;
 
-	if (CVarUseExtrinsicsGuess.GetValueOnGameThread())
+	if (CVarUseExtrinsicsGuessAruco.GetValueOnGameThread())
 	{
 		check(CameraPoses.Num() == CalibrationRows.Num());
 
