CameraCalibration: Remove coplanar points constraint on distortion aruco algo, which makes sorting markers by calibration point component unnecessary
#rnx

diff --git a/Engine/Plugins/VirtualProduction/CameraCalibration/Source/CameraCalibrationEditor/Private/Calibrators/CameraLensDistortionAlgoAruco.cpp b/Engine/Plugins/VirtualProduction/CameraCalibration/Source/CameraCalibrationEditor/Private/Calibrators/CameraLensDistortionAlgoAruco.cpp
--- a/Engine/Plugins/VirtualProduction/CameraCalibration/Source/CameraCalibrationEditor/Private/Calibrators/CameraLensDistortionAlgoAruco.cpp
+++ b/Engine/Plugins/VirtualProduction/CameraCalibration/Source/CameraCalibrationEditor/Private/Calibrators/CameraLensDistortionAlgoAruco.cpp
@@ -35,7 +35,6 @@
 #define LOCTEXT_NAMESPACE "CameraLensDistortionAlgoAruco"
 
 #if WITH_EDITOR
-static TAutoConsoleVariable<bool> CVarUseIntrinsicsGuessAruco(TEXT("LensDistortionAruco.UseIntrinsicsGuess"), true, TEXT("If true, the solver initializes the camera intrinsics to a user-provided estimate. Otherwise, the solver will compute the initial values."));
 static TAutoConsoleVariable<bool> CVarFixExtrinsicsAruco(TEXT("LensDistortionAruco.FixExtrinsics"), false, TEXT("If true, the solver will fix the camera extrinsics to the user-provided camera poses"));
 static TAutoConsoleVariable<bool> CVarFixZeroDistortionAruco(TEXT("LensDistortionAruco.FixZeroDistortion"), false, TEXT("If true, the solver will fix all distortion values to always be 0"));
 static TAutoConsoleVariable<bool> CVarUseExtrinsicsGuessAruco(TEXT("LensDistortionAruco.UseExtrinsicsGuess"), false, TEXT("If true, the actual calibrator and camera poses will be used when running the solver"));
@@ -377,16 +376,16 @@ bool UCameraLensDistortionAlgoAruco::GetLensDistortion(
 
 	ECalibrationFlags SolverFlags = ECalibrationFlags::None;
 
+	// Aruco markers may be detected anywhere in the image, and there is no guarantee that they will be coplanar. 
+	// The solver's initialization for focal length assumes all points in an image are coplanar. 
+	// Therefore, we must provide an intrinsics guess to skip this initialization step in the solver. 
+	EnumAddFlags(SolverFlags, ECalibrationFlags::UseIntrinsicGuess);
+
 	if (CVarUseExtrinsicsGuessAruco.GetValueOnGameThread())
 	{
 		EnumAddFlags(SolverFlags, ECalibrationFlags::UseExtrinsicGuess);
 	}
 
-	if (CVarUseIntrinsicsGuessAruco.GetValueOnAnyThread())
-	{
-		EnumAddFlags(SolverFlags, ECalibrationFlags::UseIntrinsicGuess);
-	}
-
 	if (CVarFixExtrinsicsAruco.GetValueOnAnyThread())
 	{
 		EnumAddFlags(SolverFlags, ECalibrationFlags::FixExtrinsics);
@@ -548,21 +547,16 @@ bool UCameraLensDistortionAlgoAruco::AddCalibrationRow(FText& OutErrorMessage)
 		}
 	}
 
-	// In the current implementation, we expect each set of 3D-2D point correspondences to be coplanar. All of the aruco markers belonging to a single calibration component should be coplanar.
-	// Therefore, we group the aruco calibration points by the calibration component they belong to, and each distinct component will contribute its own set of coplanar points.
-	TMap<UObject*, TArray<FArucoCalibrationPoint>> MarkerSets;
-	UE::CameraCalibration::Private::SortArucoCalibrationPoints(ArucoCalibrationPoints, MarkerSets);
-
 	ExportSessionData();
 
 	// For each set of coplanar markers, add a new calibration row
-	for (const TPair<UObject*, TArray<FArucoCalibrationPoint>>& MarkerSetPair : MarkerSets)
+	if (ArucoCalibrationPoints.Num() > 0)
 	{
 		TSharedPtr<FLensDistortionArucoRowData> NewRow = MakeShared<FLensDistortionArucoRowData>();
 		NewRow->Index = LensDistortionTool->AdvanceSessionRowIndex();
 
 		// Save the aruco points associated with this row 
-		NewRow->ArucoPoints = MarkerSetPair.Value;
+		NewRow->ArucoPoints = ArucoCalibrationPoints;
 
 		// Get the current pose of the selected CineCamera actor
 		NewRow->CameraPose = FTransform::Identity;
diff --git a/Engine/Plugins/VirtualProduction/CameraCalibration/Source/CameraCalibrationEditor/Private/CameraCalibrationUtilsPrivate.cpp b/Engine/Plugins/VirtualProduction/CameraCalibration/Source/CameraCalibrationEditor/Private/CameraCalibrationUtilsPrivate.cpp
--- a/Engine/Plugins/VirtualProduction/CameraCalibration/Source/CameraCalibrationEditor/Private/CameraCalibrationUtilsPrivate.cpp
+++ b/Engine/Plugins/VirtualProduction/CameraCalibration/Source/CameraCalibrationEditor/Private/CameraCalibrationUtilsPrivate.cpp
@@ -128,24 +128,6 @@ namespace UE::CameraCalibration::Private
 		}
 	}
 
-	void SortArucoCalibrationPoints(const TArray<FArucoCalibrationPoint>& ArucoCalibrationPoints, TMap<UObject*, TArray<FArucoCalibrationPoint>>& OutSortedArucoCalibrationPointMap)
-	{
-		for (const FArucoCalibrationPoint& ArucoCalibrationPoint : ArucoCalibrationPoints)
-		{
-			if (TArray<FArucoCalibrationPoint>* ArucoCalibrationPointSet = OutSortedArucoCalibrationPointMap.Find(ArucoCalibrationPoint.Owner))
-			{
-				ArucoCalibrationPointSet->Add(ArucoCalibrationPoint);
-			}
-			else
-			{
-				TArray<FArucoCalibrationPoint> NewArucoCalibrationPointSet;
-				NewArucoCalibrationPointSet.Add(ArucoCalibrationPoint);
-
-				OutSortedArucoCalibrationPointMap.Add(ArucoCalibrationPoint.Owner, NewArucoCalibrationPointSet);
-			}
-		}
-	}
-
 	bool FindArucoCalibrationPoint(const TArray<UCalibrationPointComponent*>& CalibrationComponents, EArucoDictionary ArucoDictionary, const FArucoMarker& ArucoMarker, FArucoCalibrationPoint& OutArucoCalibrationPoint)
 	{
 		// Build calibrator point name based on the detected marker
@@ -183,7 +165,6 @@ namespace UE::CameraCalibration::Private
 			// If all four corners were found, we're done. Otherwise, check the next calibration component for the four corners
 			if (FoundCorners == NumExpectedCorners)
 			{
-				OutArucoCalibrationPoint.Owner = Component;
 				return true;
 			}
 		}
diff --git a/Engine/Plugins/VirtualProduction/CameraCalibration/Source/CameraCalibrationEditor/Private/CameraCalibrationUtilsPrivate.h b/Engine/Plugins/VirtualProduction/CameraCalibration/Source/CameraCalibrationEditor/Private/CameraCalibrationUtilsPrivate.h
--- a/Engine/Plugins/VirtualProduction/CameraCalibration/Source/CameraCalibrationEditor/Private/CameraCalibrationUtilsPrivate.h
+++ b/Engine/Plugins/VirtualProduction/CameraCalibration/Source/CameraCalibrationEditor/Private/CameraCalibrationUtilsPrivate.h
@@ -28,12 +28,6 @@ struct FArucoCalibrationPoint
 	/** Name of the marker (follows the naming convention of "[DictionaryName]-[MarkerID]") */
 	UPROPERTY()
 	FString Name;
-
-	/** 
-	 * Object (almost certainly a CalibrationPointComponent) that owns this aruco point. 
-	 * Used to sort detected points that belong to different patterns in the same image 
-	 */
-	TObjectPtr<UObject> Owner = nullptr;
 };
 
 namespace UE::CameraCalibration::Private
@@ -53,9 +47,6 @@ namespace UE::CameraCalibration::Private
 	/** Find an aruco marker calibration point in one of the input calibration components that matches the input dictionary and marker ID */
 	bool FindArucoCalibrationPoint(const TArray<UCalibrationPointComponent*>& CalibrationComponents, EArucoDictionary ArucoDictionary, const FArucoMarker& ArucoMarker, FArucoCalibrationPoint& OutArucoCalibrationPoint);
 
-	/** Sort input calibration points based on owner */
-	void SortArucoCalibrationPoints(const TArray<FArucoCalibrationPoint>& ArucoCalibrationPoints, TMap<UObject*, TArray<FArucoCalibrationPoint>>& OutSortedArucoCalibrationPointMap);
-
 	/** Set every pixel in the input texture to the clear color */
 	void ClearTexture(UTexture2D* Texture, FColor ClearColor = FColor::Transparent);
 
