CameraCalibration: Constrain solver to solve a single camera pose, using the known motion between the tracked poses to transform the calibrated pose for each image.

diff --git a/Engine/Plugins/VirtualProduction/CameraCalibration/Source/CameraCalibrationEditor/Private/Calibrators/CameraCalibrationSolver.cpp b/Engine/Plugins/VirtualProduction/CameraCalibration/Source/CameraCalibrationEditor/Private/Calibrators/CameraCalibrationSolver.cpp
--- a/Engine/Plugins/VirtualProduction/CameraCalibration/Source/CameraCalibrationEditor/Private/Calibrators/CameraCalibrationSolver.cpp
+++ b/Engine/Plugins/VirtualProduction/CameraCalibration/Source/CameraCalibrationEditor/Private/Calibrators/CameraCalibrationSolver.cpp
@@ -196,7 +196,20 @@ double FCameraCalibrationSolver::CalibrateCamera(
 	// Initialize the solver with the number of parameters to solve, and the maximum number of iterations to run
 	const int NumIntrinsics = NumDistortionCoefficients + 4; // Includes Fx, Fy, Cx, and Cy
 	const int NumExtrinsics = 6; // 3 for rotation vector, 3 for translation vector
-	const int NumParamsToSolve = NumIntrinsics + (NumImages * NumExtrinsics);
+
+	// If using a guess for the extrinsic parameters, then the solver will be constrained to solve only one camera pose, and needs only one set of extrinsic parameters.
+	// Otherwise, the solver needs one set of extrinsic parameters per image.
+	int NumPosesToSolve = 0;
+	if ((EnumHasAnyFlags(SolverFlags, ECalibrationFlags::UseExtrinsicGuess)))
+	{
+		NumPosesToSolve = 1;
+	}
+	else
+	{
+		NumPosesToSolve = NumImages;
+	}
+
+	const int NumParamsToSolve = NumIntrinsics + (NumPosesToSolve * NumExtrinsics);
 
 	constexpr int NumErrors = 0;
 	constexpr int MaxIterations = 30;
@@ -244,30 +257,50 @@ double FCameraCalibrationSolver::CalibrateCamera(
 		Mask[4] = 0; // We do not want to solve for pixel aspect
 	}
 
-	// Initialize the starting guess for the camera's extrinsic parameters in each image
-	int ObjectPointsIndex = 0;
-	for (int ImageIndex = 0; ImageIndex < NumImages; ImageIndex++)
+	// Initialize the starting guess for the camera's extrinsic parameters. 
+	// If using a guess for the extrinsic parameters, then the initial guess is just the first input camera pose
+	// Otherwise, a starting pose will need to be computed for each image
+	if ((EnumHasAnyFlags(SolverFlags, ECalibrationFlags::UseExtrinsicGuess)))
+	{
+		cv::Mat Rotation = Solver.Params.rowRange(NumIntrinsics, NumIntrinsics + 3);
+		cv::Mat Translation = Solver.Params.rowRange(NumIntrinsics + 3, NumIntrinsics + 6);
+
+		FOpenCVHelper::MakeObjectVectorsFromCameraPose(InOutCameraPoses[0], Rotation, Translation);
+	}
+	else
 	{
-		const int NumImagePoints = NumPointsMat.at<int>(ImageIndex);
+		int ObjectPointsIndex = 0;
+		for (int ImageIndex = 0; ImageIndex < NumImages; ImageIndex++)
+		{
+			const int NumImagePoints = NumPointsMat.at<int>(ImageIndex);
 
-		// Get the object and image points for this image
-		cv::Mat ObjectPointsInImage = ObjectPointsMat.colRange(ObjectPointsIndex, ObjectPointsIndex + NumImagePoints);
-		cv::Mat ImagePointsInImage = ImagePointsMat.colRange(ObjectPointsIndex, ObjectPointsIndex + NumImagePoints);
+			// Get the object and image points for this image
+			cv::Mat ObjectPointsInImage = ObjectPointsMat.colRange(ObjectPointsIndex, ObjectPointsIndex + NumImagePoints);
+			cv::Mat ImagePointsInImage = ImagePointsMat.colRange(ObjectPointsIndex, ObjectPointsIndex + NumImagePoints);
 
-		ObjectPointsIndex += NumImagePoints;
+			ObjectPointsIndex += NumImagePoints;
 
-		// Get a view to parameters used by the solver for the rotation and translation vectors for this image
-		const int ExtrinsicOffset = NumIntrinsics + (ImageIndex * NumExtrinsics);
-		cv::Mat Rotation = Solver.Params.rowRange(ExtrinsicOffset, ExtrinsicOffset + 3);
-		cv::Mat Translation = Solver.Params.rowRange(ExtrinsicOffset + 3, ExtrinsicOffset + 6);
+			// Get a view to parameters used by the solver for the rotation and translation vectors for this image
+			const int ExtrinsicOffset = NumIntrinsics + (ImageIndex * NumExtrinsics);
+			cv::Mat Rotation = Solver.Params.rowRange(ExtrinsicOffset, ExtrinsicOffset + 3);
+			cv::Mat Translation = Solver.Params.rowRange(ExtrinsicOffset + 3, ExtrinsicOffset + 6);
 
-		if (!(EnumHasAnyFlags(SolverFlags, ECalibrationFlags::UseExtrinsicGuess)))
-		{
 			InitCameraExtrinsics(LensModel, ObjectPointsInImage, ImagePointsInImage, CameraMatrix, DistCoeffs, CvImageSize, Rotation, Translation, SolverFlags);
 		}
-		else
+	}
+
+	// If using a guess for the extrinsic parameters, compute the transformation from the first camera pose to each subsequent pose
+	// which represents how the camera moved between each image. The constrained solver, which only solves one camera pose, will use 
+	// this transformation to offset the pose used in the projection of points for each image. 
+	TArray<FTransform> CameraMovements;
+	if ((EnumHasAnyFlags(SolverFlags, ECalibrationFlags::UseExtrinsicGuess)))
+	{
+		CameraMovements.Reserve(InOutCameraPoses.Num());
+
+		for (const FTransform& Pose : InOutCameraPoses)
 		{
-			FOpenCVHelper::MakeObjectVectorsFromCameraPose(InOutCameraPoses[ImageIndex], Rotation, Translation);
+			const FTransform CameraMovement = InOutCameraPoses[0].Inverse() * Pose;
+			CameraMovements.Add(CameraMovement);
 		}
 	}
 
@@ -296,9 +329,19 @@ double FCameraCalibrationSolver::CalibrateCamera(
 			break;
 		}
 
+		// If using a guess for the extrinsic parameters, cache the current solver pose as a FTransform to more easily offset the camera pose for each image
+		FTransform CurrentSolverPose;
+		if ((EnumHasAnyFlags(SolverFlags, ECalibrationFlags::UseExtrinsicGuess)))
+		{
+			cv::Mat Rotation = Solver.Params.rowRange(NumIntrinsics, NumIntrinsics + 3);
+			cv::Mat Translation = Solver.Params.rowRange(NumIntrinsics + 3, NumIntrinsics + 6);
+
+			FOpenCVHelper::MakeCameraPoseFromObjectVectors(Rotation, Translation, CurrentSolverPose);
+		}
+
 		ReprojectionError = 0;
 
-		ObjectPointsIndex = 0;
+		int ObjectPointsIndex = 0;
 		for (int ImageIndex = 0; ImageIndex < NumImages; ImageIndex++)
 		{
 			int NumImagePoints = NumPointsMat.at<int>(ImageIndex);
@@ -309,10 +352,30 @@ double FCameraCalibrationSolver::CalibrateCamera(
 
 			ObjectPointsIndex += NumImagePoints;
 
-			// Get the rotation and translation vectors for this image
-			const int ExtrinsicOffset = NumIntrinsics + (ImageIndex * NumExtrinsics);
-			cv::Mat Rotation = Solver.Params.rowRange(ExtrinsicOffset, ExtrinsicOffset + 3);
-			cv::Mat Translation = Solver.Params.rowRange(ExtrinsicOffset + 3, ExtrinsicOffset + 6);
+			int ExtrinsicOffset = 0;
+			if ((EnumHasAnyFlags(SolverFlags, ECalibrationFlags::UseExtrinsicGuess)))
+			{
+				ExtrinsicOffset = NumIntrinsics;
+			}
+			else
+			{
+				ExtrinsicOffset = NumIntrinsics + (ImageIndex * NumExtrinsics);
+			}
+
+			cv::Mat Rotation;
+			cv::Mat Translation;
+			if ((EnumHasAnyFlags(SolverFlags, ECalibrationFlags::UseExtrinsicGuess)))
+			{
+				// Transform the solver's current camera pose by the camera movement to get the pose for this image
+				FTransform ImagePose = CurrentSolverPose * CameraMovements[ImageIndex];
+				FOpenCVHelper::MakeObjectVectorsFromCameraPose(ImagePose, Rotation, Translation);
+			}
+			else
+			{
+				// Get the rotation and translation vectors for this image
+				Rotation = Solver.Params.rowRange(ExtrinsicOffset, ExtrinsicOffset + 3);
+				Translation = Solver.Params.rowRange(ExtrinsicOffset + 3, ExtrinsicOffset + 6);
+			}
 
 			Jacobian.resize(NumImagePoints * 2);
 			Diffs.resize(NumImagePoints * 2);
@@ -380,7 +443,7 @@ double FCameraCalibrationSolver::CalibrateCamera(
 		OutDistCoeffs.Add(DistCoeffs.at<double>(3));
 	}
 
-	for (int ImageIndex = 0; ImageIndex < NumImages; ImageIndex++)
+	for (int ImageIndex = 0; ImageIndex < NumPosesToSolve; ImageIndex++)
 	{
 		const int ExtrinsicOffset = NumIntrinsics + (ImageIndex * NumExtrinsics);
 		cv::Mat Rotation = Solver.Params.rowRange(ExtrinsicOffset, ExtrinsicOffset + 3);
