Fix monolithic crash caused by bad static string initialization

diff --git a/Engine/Plugins/VirtualProduction/CameraCalibration/Source/CameraCalibrationEditor/Private/AssetEditor/LensDistortionTool.cpp b/Engine/Plugins/VirtualProduction/CameraCalibration/Source/CameraCalibrationEditor/Private/AssetEditor/LensDistortionTool.cpp
--- a/Engine/Plugins/VirtualProduction/CameraCalibration/Source/CameraCalibrationEditor/Private/AssetEditor/LensDistortionTool.cpp
+++ b/Engine/Plugins/VirtualProduction/CameraCalibration/Source/CameraCalibrationEditor/Private/AssetEditor/LensDistortionTool.cpp
@@ -29,7 +29,6 @@
 
 namespace UE::CameraCalibration::Private::LensDistortionTool
 {
-	static const FString SaveDir = FPaths::ProjectSavedDir() / TEXT("CameraCalibration") / TEXT("LensDistortion");
 	static const FString SessionDateTimeField(TEXT("SessionDateTime"));
 	static const FString AlgoNameField(TEXT("AlgoName"));
 }
@@ -395,7 +394,9 @@ FString ULensDistortionTool::GetSessionSaveDir() const
 	const FString DatasetPrefix = TEXT("Dataset-") + CurrentAlgo->ShortName().ToString() + TEXT("Algorithm-");
 	const FString DatasetDir = DatasetPrefix + SessionTimeString;
 
-	return LensDistortionTool::SaveDir / SessionDateString / DatasetDir;
+	const FString ProjectSaveDir = FPaths::ProjectSavedDir() / TEXT("CameraCalibration") / TEXT("LensDistortion");
+
+	return ProjectSaveDir / SessionDateString / DatasetDir;
 }
 
 FString ULensDistortionTool::GetRowFilename(int32 RowIndex) const
@@ -534,7 +535,7 @@ void ULensDistortionTool::ImportCalibrationDataset()
 	IDesktopPlatform* DesktopPlatform = FDesktopPlatformModule::Get();
 	const void* ParentWindowHandle = FSlateApplication::Get().FindBestParentWindowHandleForDialogs(nullptr);
 	const FString Title = TEXT("Import Camera Calibration Dataset");
-	const FString DefaultPath = LensDistortionTool::SaveDir;
+	const FString DefaultPath = FPaths::ProjectSavedDir() / TEXT("CameraCalibration") / TEXT("LensDistortion");;
 	const FString DefaultFile = TEXT("");
 	const FString FileTypes = TEXT("Camera Calibration Dataset|*.ucamcalib");
 	const uint32 OpenFileFlags = 0;
diff --git a/Engine/Plugins/VirtualProduction/CameraCalibration/Source/CameraCalibrationEditor/Private/AssetEditor/NodalOffsetTool.cpp b/Engine/Plugins/VirtualProduction/CameraCalibration/Source/CameraCalibrationEditor/Private/AssetEditor/NodalOffsetTool.cpp
--- a/Engine/Plugins/VirtualProduction/CameraCalibration/Source/CameraCalibrationEditor/Private/AssetEditor/NodalOffsetTool.cpp
+++ b/Engine/Plugins/VirtualProduction/CameraCalibration/Source/CameraCalibrationEditor/Private/AssetEditor/NodalOffsetTool.cpp
@@ -30,7 +30,6 @@
 
 namespace UE::CameraCalibration::Private::NodalOffsetTool
 {
-	static const FString SaveDir = FPaths::ProjectSavedDir() / TEXT("CameraCalibration") / TEXT("NodalOffset");
 	static const FString SessionDateTimeField(TEXT("SessionDateTime"));
 	static const FString AlgoNameField(TEXT("AlgoName"));
 }
@@ -327,7 +326,9 @@ FString UNodalOffsetTool::GetSessionSaveDir() const
 	const FString DatasetPrefix = TEXT("Dataset-") + NodalOffsetAlgo->ShortName().ToString() + TEXT("Algorithm-");
 	const FString DatasetDir = DatasetPrefix + SessionTimeString;
 
-	return NodalOffsetTool::SaveDir / SessionDateString / DatasetDir;
+	const FString ProjectSaveDir = FPaths::ProjectSavedDir() / TEXT("CameraCalibration") / TEXT("NodalOffset");
+
+	return ProjectSaveDir / SessionDateString / DatasetDir;
 }
 
 FString UNodalOffsetTool::GetRowFilename(int32 RowIndex) const
@@ -466,7 +467,7 @@ void UNodalOffsetTool::ImportCalibrationDataset()
 	IDesktopPlatform* DesktopPlatform = FDesktopPlatformModule::Get();
 	const void* ParentWindowHandle = FSlateApplication::Get().FindBestParentWindowHandleForDialogs(nullptr);
 	const FString Title = TEXT("Import Camera Calibration Dataset");
-	const FString DefaultPath = NodalOffsetTool::SaveDir;
+	const FString DefaultPath = FPaths::ProjectSavedDir() / TEXT("CameraCalibration") / TEXT("NodalOffset");
 	const FString DefaultFile = TEXT("");
 	const FString FileTypes = TEXT("Camera Calibration Dataset|*.ucamcalib");
 	const uint32 OpenFileFlags = 0;
