CameraCalibration: Fix incorrect warning in LensFile about mismatched filmback settings caused by bad float/double value comparison.

diff --git a/Engine/Plugins/VirtualProduction/CameraCalibration/Source/CameraCalibrationEditor/Private/AssetEditor/LensInfoStep.cpp b/Engine/Plugins/VirtualProduction/CameraCalibration/Source/CameraCalibrationEditor/Private/AssetEditor/LensInfoStep.cpp
--- a/Engine/Plugins/VirtualProduction/CameraCalibration/Source/CameraCalibrationEditor/Private/AssetEditor/LensInfoStep.cpp
+++ b/Engine/Plugins/VirtualProduction/CameraCalibration/Source/CameraCalibrationEditor/Private/AssetEditor/LensInfoStep.cpp
@@ -347,10 +347,10 @@ EVisibility ULensInfoStep::GetFilmbackWarningVisibility() const
 
 		if (LensFile && CineCameraComponent)
 		{
-			const FVector2D LensFileFilmback = LensFile->LensInfo.SensorDimensions;
-			if ((LensFile->LensInfo.SensorDimensions.X != CineCameraComponent->Filmback.SensorWidth) ||
-				(LensFile->LensInfo.SensorDimensions.Y != CineCameraComponent->Filmback.SensorHeight) ||
-				(LensFile->LensInfo.SqueezeFactor != CineCameraComponent->LensSettings.SqueezeFactor))
+			const FVector2f LensFileFilmback = FVector2f(LensFile->LensInfo.SensorDimensions.X, LensFile->LensInfo.SensorDimensions.Y);
+			if (!FMath::IsNearlyEqual(LensFileFilmback.X, CineCameraComponent->Filmback.SensorWidth) ||
+				!FMath::IsNearlyEqual(LensFileFilmback.Y, CineCameraComponent->Filmback.SensorHeight) ||
+				!FMath::IsNearlyEqual(LensFile->LensInfo.SqueezeFactor, CineCameraComponent->LensSettings.SqueezeFactor))
 			{
 				return EVisibility::Visible;
 			}
