CameraCalibration: Fix calibration component filter to display name of the correct component

diff --git a/Engine/Plugins/VirtualProduction/CameraCalibration/Source/CameraCalibrationEditor/Private/Calibrators/CameraNodalOffsetAlgoPoints.cpp b/Engine/Plugins/VirtualProduction/CameraCalibration/Source/CameraCalibrationEditor/Private/Calibrators/CameraNodalOffsetAlgoPoints.cpp
--- a/Engine/Plugins/VirtualProduction/CameraCalibration/Source/CameraCalibrationEditor/Private/Calibrators/CameraNodalOffsetAlgoPoints.cpp
+++ b/Engine/Plugins/VirtualProduction/CameraCalibration/Source/CameraCalibrationEditor/Private/Calibrators/CameraNodalOffsetAlgoPoints.cpp
@@ -1086,29 +1086,25 @@ TSharedRef<SWidget> UCameraNodalOffsetAlgoPoints::BuildCalibrationComponentMenu(
 	FMenuBuilder MenuBuilder(true, nullptr);
 	MenuBuilder.BeginSection("CalibrationComponents", LOCTEXT("CalibrationComponents", "Calibration Point Components"));
 
-	AActor* CalibratorPtr = Calibrator.Get();
-	if (CalibratorPtr)
+	if (AActor* CalibratorPtr = Calibrator.Get())
 	{
 		TArray<UCalibrationPointComponent*, TInlineAllocator<NumInlineAllocations>> CalibrationPointComponents;
 		CalibratorPtr->GetComponents(CalibrationPointComponents);
 
 		for (UCalibrationPointComponent* CalibratorComponent : CalibrationPointComponents)
 		{
-			if (USceneComponent* AttachComponent = CalibratorComponent->GetAttachParent())
-			{
-				MenuBuilder.AddMenuEntry(
-					FText::Format(LOCTEXT("ComponentLabel", "{0}"), FText::FromString(AttachComponent->GetName())),
-					FText::Format(LOCTEXT("ComponentTooltip", "{0}"), FText::FromString(AttachComponent->GetName())),
-					FSlateIcon(),
-					FUIAction(
-						FExecuteAction::CreateLambda([this, CalibratorComponent] { OnCalibrationComponentSelected(CalibratorComponent);}),
-						FCanExecuteAction(),
-						FIsActionChecked::CreateLambda([this, CalibratorComponent] { return IsCalibrationComponentSelected(CalibratorComponent); })
-					),
-					NAME_None,
-					EUserInterfaceActionType::ToggleButton
-				);
-			}
+			MenuBuilder.AddMenuEntry(
+				FText::Format(LOCTEXT("ComponentLabel", "{0}"), FText::FromString(CalibratorComponent->GetName())),
+				FText::Format(LOCTEXT("ComponentTooltip", "{0}"), FText::FromString(CalibratorComponent->GetName())),
+				FSlateIcon(),
+				FUIAction(
+					FExecuteAction::CreateLambda([this, CalibratorComponent] { OnCalibrationComponentSelected(CalibratorComponent);}),
+					FCanExecuteAction(),
+					FIsActionChecked::CreateLambda([this, CalibratorComponent] { return IsCalibrationComponentSelected(CalibratorComponent); })
+				),
+				NAME_None,
+				EUserInterfaceActionType::ToggleButton
+			);
 		}
 	}
 	else
@@ -1149,14 +1145,15 @@ TSharedRef<SWidget> UCameraNodalOffsetAlgoPoints::BuildCalibrationComponentPicke
 					{
 						return LOCTEXT("MultipleCalibrationComponents", "Multiple Values");
 					}
-					else if (ActiveCalibratorComponents.Num() == 1 && ActiveCalibratorComponents[0].Get() && ActiveCalibratorComponents[0]->GetAttachParent())
+					else if (ActiveCalibratorComponents.Num() == 1)
 					{
-						return FText::FromString(ActiveCalibratorComponents[0]->GetAttachParent()->GetName());
-					}
-					else
-					{
-						return LOCTEXT("NoCalibrationComponents", "None");
+						if (const UCalibrationPointComponent* CalibratorComponent = ActiveCalibratorComponents[0].Get())
+						{
+							return FText::FromString(CalibratorComponent->GetName());
+						}
 					}
+
+					return LOCTEXT("NoCalibrationComponents", "None");
 				})
 			]
 		]; 
