CameraCalibration: Add framework for launching a distortion calibration task asynchronously. Add stub for python neural-net distortion solver.

diff --git a/Engine/Plugins/VirtualProduction/CameraCalibration/Source/CameraCalibrationEditor/Private/AssetEditor/LensDistortionTool.cpp b/Engine/Plugins/VirtualProduction/CameraCalibration/Source/CameraCalibrationEditor/Private/AssetEditor/LensDistortionTool.cpp
--- a/Engine/Plugins/VirtualProduction/CameraCalibration/Source/CameraCalibrationEditor/Private/AssetEditor/LensDistortionTool.cpp
+++ b/Engine/Plugins/VirtualProduction/CameraCalibration/Source/CameraCalibrationEditor/Private/AssetEditor/LensDistortionTool.cpp
@@ -2,6 +2,7 @@
 
 #include "LensDistortionTool.h"
 
+#include "AssetToolsModule.h"
 #include "CameraCalibrationEditorLog.h"
 #include "CameraCalibrationSettings.h"
 #include "CameraCalibrationStepsController.h"
@@ -65,6 +66,8 @@ void ULensDistortionTool::Initialize(TWeakPtr<FCameraCalibrationStepsController>
 			UpdateAlgoMap(LensFile->LensInfo.LensModel);
 		}
 	}
+
+	BuildProgressWindowWidgets();
 }
 
 void ULensDistortionTool::OnLensModelChanged(const TSubclassOf<ULensModel>& LensModel)
@@ -115,6 +118,35 @@ void ULensDistortionTool::Tick(float DeltaTime)
 	{
 		CurrentAlgo->Tick(DeltaTime);
 	}
+
+	// A valid task handle implies that there is an asynchronous calibration happening on another thread.
+	// The tool will poll the task to determine when it has finished so that the results can be saved.
+	if (CalibrationTask.IsValid())
+	{
+		if (CalibrationTask.IsCompleted())
+		{
+			// Extract the return value from the task and release the task resource
+			CalibrationResult = CalibrationTask.GetResult();
+			CalibrationTask = {};
+
+			if (!CalibrationResult.ErrorMessage.IsEmpty())
+			{
+				const FText Message = FText::Format(LOCTEXT("CalibrationErrorResult", "Calibration Error: {0}"), CalibrationResult.ErrorMessage);
+				ProgressTextWidget->SetText(Message);
+			}
+			else
+			{
+				// Update progress window with final reprojection error
+				FFormatOrderedArguments Arguments;
+				Arguments.Add(FText::FromString(FString::Printf(TEXT("%.3f"), CalibrationResult.ReprojectionError)));
+
+				const FText Message = FText::Format(LOCTEXT("CalibrationTaskResult", "Reprojection Error: {0} pixels"), Arguments);
+				ProgressTextWidget->SetText(Message);
+			}
+
+			OkayButton->SetEnabled(true);
+		}
+	}
 }
 
 bool ULensDistortionTool::OnViewportClicked(const FGeometry& MyGeometry, const FPointerEvent& MouseEvent)
@@ -138,6 +170,65 @@ TSharedRef<SWidget> ULensDistortionTool::BuildUI()
 	return DistortionWidget.ToSharedRef();
 }
 
+void ULensDistortionTool::BuildProgressWindowWidgets()
+{
+	ProgressWindow = SNew(SWindow)
+		.Title(LOCTEXT("ProgressWindowTitle", "Distortion Calibration Progress"))
+		.SizingRule(ESizingRule::Autosized)
+		.HasCloseButton(false)
+		.SupportsMaximize(false)
+		.SupportsMinimize(true);
+
+	ProgressTextWidget = SNew(STextBlock).Text(FText::GetEmpty());
+
+	OkayButton = SNew(SButton)
+		.IsEnabled(false)
+		.HAlign(HAlign_Center)
+		.VAlign(VAlign_Center)
+		.Text(LOCTEXT("OkText", "Ok"))
+		.OnClicked_UObject(this, &ULensDistortionTool::OnOkPressed);
+
+	TSharedRef<SWidget> WindowContent = SNew(SVerticalBox)
+
+		// Text widget to display the current progress of the calibration
+		+ SVerticalBox::Slot()
+		.HAlign(EHorizontalAlignment::HAlign_Center)
+		.VAlign(EVerticalAlignment::VAlign_Center)
+		[
+			ProgressTextWidget.ToSharedRef()
+		]
+
+		// Ok and Cancel buttons
+		+ SVerticalBox::Slot()
+		.HAlign(EHorizontalAlignment::HAlign_Center)
+		.VAlign(EVerticalAlignment::VAlign_Center)
+		[
+			SNew(SHorizontalBox)
+
+			+ SHorizontalBox::Slot()
+			.AutoWidth()
+			[
+				OkayButton.ToSharedRef()
+			]
+
+			+ SHorizontalBox::Slot()
+			.AutoWidth()
+			[
+				SNew(SButton)
+				.Text(LOCTEXT("CancelText", "Cancel"))
+				.HAlign(HAlign_Center)
+				.VAlign(VAlign_Center)
+				.OnClicked_UObject(this, &ULensDistortionTool::OnCancelPressed)
+			]
+		];
+
+	ProgressWindow->SetContent(WindowContent);
+
+	// Create the window, but start with it hidden. When the user initiates a calibration, the progress window will be shown.
+	FSlateApplication::Get().AddWindow(ProgressWindow.ToSharedRef());
+	ProgressWindow->HideWindow();
+}
+
 bool ULensDistortionTool::DependsOnStep(UCameraCalibrationStep* Step) const
 {
 	return Cast<ULensInfoStep>(Step) != nullptr;
@@ -243,80 +334,142 @@ UCameraLensDistortionAlgo* ULensDistortionTool::GetAlgo() const
 
 void ULensDistortionTool::OnSaveCurrentCalibrationData()
 {
-	if (!CameraCalibrationStepsController.IsValid())
+	if (!CurrentAlgo)
 	{
 		return;
 	}
 
-	UCameraLensDistortionAlgo* Algo = GetAlgo();
-	
-	if (!Algo)
-	{
-		FText ErrorMessage = LOCTEXT("NoAlgoFound", "No algo found");
-		FMessageDialog::Open(EAppMsgType::Ok, ErrorMessage);
-		return;
-	}
-	
-	ULensFile* LensFile = CameraCalibrationStepsController.Pin()->GetLensFile();
+	const FText TitleError = LOCTEXT("LensCalibrationError", "Lens Calibration Error");
+	const FText UnknownError = LOCTEXT("UnknownError", "An error occurred initiating the distortion calibration. Check the output log for details.");
 
-	if (!LensFile)
+	if (CurrentAlgo->SupportsAsyncCalibration())
 	{
-		FText ErrorMessage = LOCTEXT("NoLensFile", "No Lens File");
-		FMessageDialog::Open(EAppMsgType::Ok, ErrorMessage);
-		return;
-	}
+		CalibrationTask = CurrentAlgo->BeginCalibration();
 
-	const FText TitleError = LOCTEXT("LensCalibrationError", "Lens Calibration Error");
-	const FText TitleInfo = LOCTEXT("LensCalibrationInfo", "Lens Calibration Info");
+		if (!CalibrationTask.IsValid())
+		{
+			FMessageDialog::Open(EAppMsgType::Ok, UnknownError, TitleError);
+			return;
+		}
 
-	float Focus;
-	float Zoom;
-	FDistortionInfo DistortionInfo;
-	FFocalLengthInfo FocalLengthInfo;
-	FImageCenterInfo ImageCenterInfo;
-	TSubclassOf<ULensModel> LensModel;
-	double Error;
+		// TODO: This text is temporary. Update to register a delegate with the algo to provide implementation-specific status text
+		ProgressTextWidget->SetText(LOCTEXT("CalibrationProgressText", "Calibrating Lens Distortion..."));
 
-	// Get distortion value, and if errors, inform the user.
+		// Ensure that the Ok button is disabled and show the progress window
+		OkayButton->SetEnabled(false);
+		ProgressWindow->ShowWindow();
+	}
+	else
 	{
-		FText ErrorMessage;
+		TSubclassOf<ULensModel> LensModel;
+		bool bResult = CurrentAlgo->GetLensDistortion(
+			CalibrationResult.EvaluatedFocus, 
+			CalibrationResult.EvaluatedZoom, 
+			CalibrationResult.Parameters, 
+			CalibrationResult.FocalLength, 
+			CalibrationResult.ImageCenter, 
+			LensModel, 
+			CalibrationResult.ReprojectionError, 
+			CalibrationResult.ErrorMessage);
 
-		if (!Algo->GetLensDistortion(Focus, Zoom, DistortionInfo, FocalLengthInfo, ImageCenterInfo, LensModel, Error, ErrorMessage))
+		if (!bResult)
 		{
-			FMessageDialog::Open(EAppMsgType::Ok, ErrorMessage, TitleError);
+			if (CalibrationResult.ErrorMessage.IsEmpty())
+			{
+				CalibrationResult.ErrorMessage = UnknownError;
+			}
+
+			FMessageDialog::Open(EAppMsgType::Ok, CalibrationResult.ErrorMessage, TitleError);
 			return;
 		}
-	}
 
-	// Show reprojection error
-	{
+		// Update progress window with final reprojection error
+		ProgressWindow->ShowWindow();
+
 		FFormatOrderedArguments Arguments;
-		Arguments.Add(FText::FromString(FString::Printf(TEXT("%.2f"), Error)));
+		Arguments.Add(FText::FromString(FString::Printf(TEXT("%.3f"), CalibrationResult.ReprojectionError)));
 
-		const FText Message = FText::Format(LOCTEXT("ReprojectionError", "RMS Reprojection Error: {0} pixels"), Arguments);
+		const FText Message = FText::Format(LOCTEXT("ReprojectionError", "Reprojection Error: {0} pixels"), Arguments);
+		ProgressTextWidget->SetText(Message);
 
-		// Allow the user to cancel adding to the LUT if the reprojection error is unacceptable.
-		if (FMessageDialog::Open(EAppMsgType::OkCancel, Message, TitleInfo) != EAppReturnType::Ok)
-		{
-			return;
-		}
+		OkayButton->SetEnabled(true);
 	}
+}
+
+FReply ULensDistortionTool::OnCancelPressed()
+{
+	// TODO: Implement correct cancel handling. Send a signal to the algo to interrupt the currently running calibration task 
+
+	return FReply::Handled();
+}
+
+FReply ULensDistortionTool::OnOkPressed()
+{
+	SaveCalibrationResult();
+
+	ProgressWindow->HideWindow();
+
+	return FReply::Handled();
+}
+
+void ULensDistortionTool::SaveCalibrationResult()
+{
+	ULensFile* LensFile = CameraCalibrationStepsController.Pin()->GetLensFile();
 
-	if (LensFile->HasSamples(ELensDataCategory::Distortion) && LensFile->LensInfo.LensModel != LensModel)
+	if (!LensFile)
 	{
-		const FText ErrorMessage = LOCTEXT("LensDistortionModelMismatch", "There is a distortion model mismatch between the new and existing samples");
-		FMessageDialog::Open(EAppMsgType::Ok, ErrorMessage, TitleError);
 		return;
 	}
 
-	FScopedTransaction Transaction(LOCTEXT("SaveCurrentDistortionData", "Save Current Distortion Data"));
+	// If the calibration result contains the name of an ST Map file on disk instead of a UTexture, then we attempt to import it for the user
+	if (!CalibrationResult.STMap.DistortionMap && !CalibrationResult.STMapFullPath.IsEmpty())
+	{
+		FAssetToolsModule& AssetToolsModule = FModuleManager::GetModuleChecked<FAssetToolsModule>("AssetTools");
+
+		TArray<FString> TextureFileNames;
+		TextureFileNames.Add(CalibrationResult.STMapFullPath);
+		TArray<UObject*> ImportedImages = AssetToolsModule.Get().ImportAssets(TextureFileNames, FPaths::ProjectContentDir());
+
+		CalibrationResult.STMap.DistortionMap = (ImportedImages.Num() > 0) ? Cast<UTexture>(ImportedImages[0]) : nullptr;
+	}
+
+	// Depending on the algo, it is possible that the result feature calibrated distortion parameters or an ST Map.
+	// If the result contains any distortion parameters, then the results will be written as a distortion point in the Lens File
+	// Otherwise, if the result contains a valid ST Map, then it will be added to the Lens File
+	if (CalibrationResult.Parameters.Parameters.Num() > 0)
+	{
+		if (LensFile->DataMode != ELensDataMode::Parameters)
+		{
+			LensFile->DataMode = ELensDataMode::Parameters;
+			UE_LOG(LogCameraCalibrationEditor, Log, TEXT("The LensFile's data mode was set to ST Map, but the latest calibration result returned distortion parameters. Data mode will change to Parameters."));
+		}
+
+		FScopedTransaction Transaction(LOCTEXT("SaveCurrentDistortionData", "Save Current Distortion Data"));
+		LensFile->Modify();
 
- 	LensFile->Modify();
+		LensFile->AddDistortionPoint(CalibrationResult.EvaluatedFocus, CalibrationResult.EvaluatedZoom, CalibrationResult.Parameters, CalibrationResult.FocalLength);
+		LensFile->AddImageCenterPoint(CalibrationResult.EvaluatedFocus, CalibrationResult.EvaluatedZoom, CalibrationResult.ImageCenter);
+	}
+	else if (CalibrationResult.STMap.DistortionMap)
+	{
+		if (LensFile->DataMode != ELensDataMode::STMap)
+		{
+			LensFile->DataMode = ELensDataMode::STMap;
+			UE_LOG(LogCameraCalibrationEditor, Log, TEXT("The LensFile's data mode was set to Parameters, but the latest calibration result returned an ST Map. Data mode will change to ST Map."));
+		}
 
-	LensFile->AddDistortionPoint(Focus, Zoom, DistortionInfo, FocalLengthInfo);
-	LensFile->AddImageCenterPoint(Focus, Zoom, ImageCenterInfo);
+		FScopedTransaction Transaction(LOCTEXT("SaveCurrentDistortionData", "Save Current Distortion Data"));
+		LensFile->Modify();
 
-	Algo->OnDistortionSavedToLens();
+		LensFile->AddSTMapPoint(CalibrationResult.EvaluatedFocus, CalibrationResult.EvaluatedZoom, CalibrationResult.STMap);
+		LensFile->AddFocalLengthPoint(CalibrationResult.EvaluatedFocus, CalibrationResult.EvaluatedZoom, CalibrationResult.FocalLength);
+		LensFile->AddImageCenterPoint(CalibrationResult.EvaluatedFocus, CalibrationResult.EvaluatedZoom, CalibrationResult.ImageCenter);
+	}
+
+	if (UCameraLensDistortionAlgo* Algo = GetAlgo())
+	{
+		Algo->OnDistortionSavedToLens();
+	}
 }
 
 FCameraCalibrationStepsController* ULensDistortionTool::GetCameraCalibrationStepsController() const
diff --git a/Engine/Plugins/VirtualProduction/CameraCalibration/Source/CameraCalibrationEditor/Private/AssetEditor/LensDistortionTool.h b/Engine/Plugins/VirtualProduction/CameraCalibration/Source/CameraCalibrationEditor/Private/AssetEditor/LensDistortionTool.h
--- a/Engine/Plugins/VirtualProduction/CameraCalibration/Source/CameraCalibrationEditor/Private/AssetEditor/LensDistortionTool.h
+++ b/Engine/Plugins/VirtualProduction/CameraCalibration/Source/CameraCalibrationEditor/Private/AssetEditor/LensDistortionTool.h
@@ -4,10 +4,13 @@
 
 #include "CoreMinimal.h"
 
-#include "Engine/World.h"
 #include "CameraCalibrationStep.h"
+#include "CameraCalibrationTypes.h"
+#include "CameraLensDistortionAlgo.h"
+#include "Engine/World.h"
 #include "ImageCore.h"
 #include "SLensDistortionToolPanel.h"
+#include "Widgets/Input/SButton.h"
 
 #include "LensDistortionTool.generated.h"
 
@@ -106,6 +109,18 @@ private:
 	/** Update the list of algos that support the current Lens Model */
 	void UpdateAlgoMap(const TSubclassOf<ULensModel>& LensModel);
 
+	/** Build the UI for the progress window, including the window content and buttons */
+	void BuildProgressWindowWidgets();
+
+	/** Called in response to user clicking "Cancel" button in the progress window to cancel a currently running calibration task */
+	FReply OnCancelPressed();
+
+	/** Called in response to user clicking "Ok" button in the progress window to save the result from a finished calibration task */
+	FReply OnOkPressed();
+
+	/** Save the results from a finished distortion calibration to the Lens File */
+	void SaveCalibrationResult();
+
 public:
 	/** Stores info about the current calibration session of this tool */
 	FLensDistortionSessionInfo SessionInfo;
@@ -136,4 +151,19 @@ private:
 
 	/** UI Widget for this Tool */
 	TSharedPtr<SLensDistortionToolPanel> DistortionWidget;
+
+	/** UI Widget for the calibration task progress window */
+	TSharedPtr<SWindow> ProgressWindow;
+
+	/** UI Widget for the text to be displayed in the progress window */
+	TSharedPtr<STextBlock> ProgressTextWidget;
+
+	/** UI Widget for the "Ok" button in the progress window */
+	TSharedPtr<SButton> OkayButton;
+
+	/** An asynchronous task handle. When valid, the tool will poll its state to determine when the task has completed, and then extract the calibration result from this task handle. */
+	FDistortionCalibrationTask CalibrationTask;
+
+	/** The result from the most recently completed distortion calibration. It is saved as a member variable because the result is not saved immediately, but only if the user confirms that it should be saved. */
+	FDistortionCalibrationResult CalibrationResult;
 };
diff --git a/Engine/Plugins/VirtualProduction/CameraCalibration/Source/CameraCalibrationEditor/Private/Calibrators/CameraCalibrationSolver.cpp b/Engine/Plugins/VirtualProduction/CameraCalibration/Source/CameraCalibrationEditor/Private/Calibrators/CameraCalibrationSolver.cpp
--- a/Engine/Plugins/VirtualProduction/CameraCalibration/Source/CameraCalibrationEditor/Private/Calibrators/CameraCalibrationSolver.cpp
+++ b/Engine/Plugins/VirtualProduction/CameraCalibration/Source/CameraCalibrationEditor/Private/Calibrators/CameraCalibrationSolver.cpp
@@ -23,8 +23,8 @@ DEFINE_LOG_CATEGORY_STATIC(LogCameraCalibrationSolver, Log, All);
 
 double FCameraCalibrationSolver::CalibrateCamera(
 	const TSubclassOf<ULensModel> LensModel,
-	const TArray<TArray<FVector>>& InObjectPoints,
-	const TArray<TArray<FVector2f>>& InImagePoints,
+	const TArray<FObjectPoints>& InObjectPointsArray,
+	const TArray<FImagePoints>& InImagePointsArray,
 	const FIntPoint ImageSize,
 	FVector2D& InOutFocalLength,
 	FVector2D& InOutImageCenter,
@@ -34,7 +34,7 @@ double FCameraCalibrationSolver::CalibrateCamera(
 	ECalibrationFlags SolverFlags)
 {
 #if WITH_OPENCV
-	const int NumImages = InObjectPoints.Num();
+	const int NumImages = InObjectPointsArray.Num();
 
 	// Create an array to store the number of points in each image
 	cv::Mat NumPointsMat = cv::Mat(1, NumImages, CV_32S);
@@ -43,7 +43,7 @@ double FCameraCalibrationSolver::CalibrateCamera(
 	int MaxPoints = 0;
 	for (int ImageIndex = 0; ImageIndex < NumImages; ++ImageIndex)
 	{
-		const int NumPointsInImage = InObjectPoints[ImageIndex].Num();
+		const int NumPointsInImage = InObjectPointsArray[ImageIndex].Points.Num();
 
 		NumPointsMat.at<int>(ImageIndex) = NumPointsInImage;
 		NumTotalPoints += NumPointsInImage;
@@ -51,29 +51,11 @@ double FCameraCalibrationSolver::CalibrateCamera(
 		MaxPoints = MAX(MaxPoints, NumPointsInImage);
 	}
 
-	// Convert the object points from Unreal's coordinate system to OpenCV's
-	TArray<TArray<FVector>> CvObjectPoints;
-	CvObjectPoints.Reserve(NumImages);
-	for (int ImageIndex = 0; ImageIndex < NumImages; ++ImageIndex)
-	{
-		const TArray<FVector>& ObjectPointsForImage = InObjectPoints[ImageIndex];
-
-		TArray<FVector> CvObjectPointsForImage;
-		CvObjectPointsForImage.Reserve(ObjectPointsForImage.Num());
-
-		for (const FVector& Point : ObjectPointsForImage)
-		{
-			CvObjectPointsForImage.Add(FOpenCVHelper::ConvertUnrealToOpenCV(Point));
-		}
-
-		CvObjectPoints.Add(CvObjectPointsForImage);
-	}
-
 	cv::Mat ObjectPointsMat = cv::Mat(1, NumTotalPoints, CV_64FC3);
 	cv::Mat ImagePointsMat = cv::Mat(1, NumTotalPoints, CV_64FC2);
 
 	// Reorganize the 3D and 2D points from the input arrays or arrays to be laid out linearly in memory in two cv::Mat objects
-	GatherPoints(CvObjectPoints, InImagePoints, ObjectPointsMat, ImagePointsMat);
+	GatherPoints(InObjectPointsArray, InImagePointsArray, ObjectPointsMat, ImagePointsMat);
 
 	double RMSE = 0.0;
 
@@ -114,28 +96,29 @@ double FCameraCalibrationSolver::CalibrateCamera(
 		std::vector<std::vector<cv::Point2f>> Samples2d;
 		std::vector<std::vector<cv::Point3f>> Samples3d;
 
-		Samples2d.reserve(InImagePoints.Num());
-		Samples3d.reserve(CvObjectPoints.Num());
+		Samples2d.reserve(InImagePointsArray.Num());
+		Samples3d.reserve(InObjectPointsArray.Num());
 
-		for (const TArray<FVector>& Image : CvObjectPoints)
+		for (const FObjectPoints& PointsInImage : InObjectPointsArray)
 		{
 			std::vector<cv::Point3f> Points3d;
-			Points3d.reserve(Image.Num());
+			Points3d.reserve(PointsInImage.Points.Num());
 
-			for (const FVector& Point3d : Image)
+			for (const FVector& Point3d : PointsInImage.Points)
 			{
-				Points3d.push_back(cv::Point3f(Point3d.X, Point3d.Y, Point3d.Z));
+				const FVector Point3dCV = FOpenCVHelper::ConvertUnrealToOpenCV(Point3d);
+				Points3d.push_back(cv::Point3f(Point3dCV.X, Point3dCV.Y, Point3dCV.Z));
 			}
 
 			Samples3d.push_back(Points3d);
 		}
 
-		for (const TArray<FVector2f>& Image : InImagePoints)
+		for (const FImagePoints& PointsInImage : InImagePointsArray)
 		{
 			std::vector<cv::Point2f> Points2d;
-			Points2d.reserve(Image.Num());
+			Points2d.reserve(PointsInImage.Points.Num());
 
-			for (const FVector2f& Point2d : Image)
+			for (const FVector2D& Point2d : PointsInImage.Points)
 			{
 				Points2d.push_back(cv::Point2f(Point2d.X, Point2d.Y));
 			}
@@ -1387,8 +1370,8 @@ void FCameraCalibrationSolver::ProjectPointsSpherical(
 }
 
 void FCameraCalibrationSolver::GatherPoints(
-	const TArray<TArray<FVector>>& ObjectPoints,
-	const TArray<TArray<FVector2f>>& ImagePoints,
+	const TArray<FObjectPoints>& InObjectPointsArray,
+	const TArray<FImagePoints>& InImagePointsArray,
 	cv::Mat& ObjectPointsMat,
 	cv::Mat& ImagePointsMat)
 {
@@ -1397,19 +1380,21 @@ void FCameraCalibrationSolver::GatherPoints(
 
 	int TotalPointIndex = 0;
 
-	const int NumImages = ObjectPoints.Num();
+	const int NumImages = InObjectPointsArray.Num();
 
 	for (int ImageIndex = 0; ImageIndex < NumImages; ++ImageIndex)
 	{
-		const TArray<FVector>& CurrentObjectPoints = ObjectPoints[ImageIndex];
-		const TArray<FVector2f>& CurrentImagePoints = ImagePoints[ImageIndex];
+		const FObjectPoints& CurrentObjectPoints = InObjectPointsArray[ImageIndex];
+		const FImagePoints& CurrentImagePoints = InImagePointsArray[ImageIndex];
 
-		const int NumPointsInImage = CurrentObjectPoints.Num();
+		const int NumPointsInImage = CurrentObjectPoints.Points.Num();
 
 		for (int PointIndex = 0; PointIndex < NumPointsInImage; ++PointIndex)
 		{
-			ObjectPointsMatData[TotalPointIndex + PointIndex] = cv::Point3d(CurrentObjectPoints[PointIndex].X, CurrentObjectPoints[PointIndex].Y, CurrentObjectPoints[PointIndex].Z);
-			ImagePointsMatData[TotalPointIndex + PointIndex] = cv::Point2d(CurrentImagePoints[PointIndex].X, CurrentImagePoints[PointIndex].Y);
+			const FVector ObjectPointCV = FOpenCVHelper::ConvertUnrealToOpenCV(CurrentObjectPoints.Points[PointIndex]);
+			ObjectPointsMatData[TotalPointIndex + PointIndex] = cv::Point3d(ObjectPointCV.X, ObjectPointCV.Y, ObjectPointCV.Z);
+
+			ImagePointsMatData[TotalPointIndex + PointIndex] = cv::Point2d(CurrentImagePoints.Points[PointIndex].X, CurrentImagePoints.Points[PointIndex].Y);
 		}
 
 		TotalPointIndex += NumPointsInImage;
diff --git a/Engine/Plugins/VirtualProduction/CameraCalibration/Source/CameraCalibrationEditor/Private/Calibrators/CameraCalibrationSolver.h b/Engine/Plugins/VirtualProduction/CameraCalibration/Source/CameraCalibrationEditor/Private/Calibrators/CameraCalibrationSolver.h
--- a/Engine/Plugins/VirtualProduction/CameraCalibration/Source/CameraCalibrationEditor/Private/Calibrators/CameraCalibrationSolver.h
+++ b/Engine/Plugins/VirtualProduction/CameraCalibration/Source/CameraCalibrationEditor/Private/Calibrators/CameraCalibrationSolver.h
@@ -13,6 +13,55 @@
 
 #endif	// WITH_OPENCV
 
+#include "CameraCalibrationSolver.generated.h"
+
+/** 
+ * An array of 3D object points associated with a single calibration image.
+ * Structure is needed because Blueprints and Python do not support arrays of arrays.
+ */
+USTRUCT(BlueprintType)
+struct FObjectPoints
+{
+	GENERATED_BODY()
+
+	/** 3D object points in an image */
+	UPROPERTY(BlueprintReadWrite, Category = "Calibration")
+	TArray<FVector> Points;
+};
+
+/** 
+ * An array of 2D image points associated with a single calibration image 
+ * Structure is needed because Blueprints and Python do not support arrays of arrays.
+ */
+USTRUCT(BlueprintType)
+struct FImagePoints
+{
+	GENERATED_BODY()
+
+	/** 2D image points in an image */
+	UPROPERTY(BlueprintReadWrite, Category = "Calibration")
+	TArray<FVector2D> Points;
+};
+
+/** 
+ * Base lens distortion solver class that can be inherited from in Blueprints or Python 
+ */
+UCLASS(Blueprintable)
+class ULensDistortionSolver : public UObject
+{
+	GENERATED_BODY()
+
+public:
+	/** Calibrate camera intrinsics and distortion from a set of input 3D-2D point correspondences and initial camera intrinsics guess. */
+	UFUNCTION(BlueprintImplementableEvent, Category = "Calibration")
+	FDistortionCalibrationResult Solve(
+		const TArray<FObjectPoints>& ObjectPointArray,
+		const TArray<FImagePoints>& ImagePointArray,
+		const FIntPoint ImageSize,
+		const FVector2D& FocalLength,
+		const FVector2D& ImageCenter) const;
+};
+
 /**
   * Flags used to modify the execution of the calibration solver
   */
@@ -40,8 +89,8 @@ public:
 	  */
 	static double CalibrateCamera(
 		const TSubclassOf<ULensModel> LensModel,
-		const TArray<TArray<FVector>>& InObjectPoints,
-		const TArray<TArray<FVector2f>>& InImagePoints,
+		const TArray<FObjectPoints>& InObjectPointsArray,
+		const TArray<FImagePoints>& InImagePointsArray,
 		const FIntPoint ImageSize,
 		FVector2D& InOutFocalLength,
 		FVector2D& InOutImageCenter,
@@ -130,8 +179,8 @@ private:
 
 	/** Copy the input 3D and 2D points to OpenCV matrices for ease of use with the solver */
 	static void GatherPoints(
-		const TArray<TArray<FVector>>& ObjectPoints,
-		const TArray<TArray<FVector2f>>& ImagePoints,
+		const TArray<FObjectPoints>& InObjectPointsArray,
+		const TArray<FImagePoints>& InImagePointsArray,
 		cv::Mat& ObjectPointsMat,
 		cv::Mat& ImagePointsMat);
 
diff --git a/Engine/Plugins/VirtualProduction/CameraCalibration/Source/CameraCalibrationEditor/Private/Calibrators/CameraLensDistortionAlgoAruco.cpp b/Engine/Plugins/VirtualProduction/CameraCalibration/Source/CameraCalibrationEditor/Private/Calibrators/CameraLensDistortionAlgoAruco.cpp
--- a/Engine/Plugins/VirtualProduction/CameraCalibration/Source/CameraCalibrationEditor/Private/Calibrators/CameraLensDistortionAlgoAruco.cpp
+++ b/Engine/Plugins/VirtualProduction/CameraCalibration/Source/CameraCalibrationEditor/Private/Calibrators/CameraLensDistortionAlgoAruco.cpp
@@ -38,6 +38,7 @@
 static TAutoConsoleVariable<bool> CVarFixExtrinsicsAruco(TEXT("LensDistortionAruco.FixExtrinsics"), false, TEXT("If true, the solver will fix the camera extrinsics to the user-provided camera poses"));
 static TAutoConsoleVariable<bool> CVarFixZeroDistortionAruco(TEXT("LensDistortionAruco.FixZeroDistortion"), false, TEXT("If true, the solver will fix all distortion values to always be 0"));
 static TAutoConsoleVariable<bool> CVarUseExtrinsicsGuessAruco(TEXT("LensDistortionAruco.UseExtrinsicsGuess"), false, TEXT("If true, the actual calibrator and camera poses will be used when running the solver"));
+static TAutoConsoleVariable<bool> CVarUseNeuralNetSolverAruco(TEXT("LensDistortionAruco.UseNeuralNetSolver"), false, TEXT("If true, the neural network based solver will solver for distortion and camera intrinsics. Otherwise, the traditional OpenCV solver will be used."));
 #endif
 
 const int UCameraLensDistortionAlgoAruco::DATASET_VERSION = 1;
@@ -264,47 +265,38 @@ bool UCameraLensDistortionAlgoAruco::OnViewportClicked(const FGeometry& MyGeomet
 	return true;
 }
 
-bool UCameraLensDistortionAlgoAruco::GetLensDistortion(
-	float& OutFocus,
-	float& OutZoom,
-	FDistortionInfo& OutDistortionInfo,
-	FFocalLengthInfo& OutFocalLengthInfo,
-	FImageCenterInfo& OutImageCenterInfo,
-	TSubclassOf<ULensModel>& OutLensModel,
-	double& OutError,
-	FText& OutErrorMessage)
+FDistortionCalibrationTask UCameraLensDistortionAlgoAruco::BeginCalibration()
 {
+	FDistortionCalibrationTask CalibrationTask = {};
+
 	if (CalibrationRows.Num() < 1)
 	{
-		OutErrorMessage = LOCTEXT("NotEnoughSamples", "At least 1 calibration row is required");
-		return false;
+		UE_LOG(LogCameraCalibrationEditor, Error, TEXT("Could not initiate distortion calibration. At least 1 calibration row is required."));
+		return CalibrationTask;
 	}
 
 	ULensDistortionTool* LensDistortionTool = WeakTool.Get();
 	if (!LensDistortionTool)
 	{
-		return false;
+		return CalibrationTask;
 	}
 
 	FCameraCalibrationStepsController* StepsController = LensDistortionTool->GetCameraCalibrationStepsController();
 	if (!StepsController)
 	{
-		return false;
+		return CalibrationTask;
 	}
 
 	ULensFile* LensFile = StepsController->GetLensFile();
 	if (!LensFile)
 	{
-		OutErrorMessage = LOCTEXT("LensFileNotFound", "LensFile not found");
-		return false;
+		return CalibrationTask;
 	}
 
-	// Only parameters data mode supported at the moment
-	if (LensFile->DataMode != ELensDataMode::Parameters)
-	{
-		OutErrorMessage = LOCTEXT("OnlyParametersDataModeSupported", "Only Parameters Data Mode supported");
-		return false;
-	}
+	// The evaluation focus and zoom values should not change during data collection, so all rows should have the same values.
+	const FLensFileEvaluationInputs LensFileEvalInputs = StepsController->GetLensFileEvaluationInputs();
+	const float Focus = LensFileEvalInputs.Focus;
+	const float Zoom = LensFileEvalInputs.Zoom;
 
 	// Initialize the focal length estimate in pixels
 	const FIntPoint ImageSize = LensFile->CameraFeedInfo.GetDimensions();
@@ -313,8 +305,8 @@ bool UCameraLensDistortionAlgoAruco::GetLensDistortion(
 
 	if (FMath::IsNearlyZero(PixelAspect))
 	{
-		OutErrorMessage = LOCTEXT("PixelAspectNearlyZero", "The pixel aspect in the Lens Information is zero, which is invalid.");
-		return false;
+		UE_LOG(LogCameraCalibrationEditor, Error, TEXT("Could not initiate distortion calibration. The pixel aspect ratio of the CineCamera is zero, which is invalid."));
+		return CalibrationTask;
 	}
 
 	const float PhysicalSensorWidth = StepsController->GetLensFileEvaluationInputs().Filmback.SensorWidth;
@@ -322,14 +314,14 @@ bool UCameraLensDistortionAlgoAruco::GetLensDistortion(
 
 	if (FMath::IsNearlyZero(DesqueezeSensorWidth))
 	{
-		OutErrorMessage = LOCTEXT("SensorWidthNearlyZero", "The sensor width of the CineCamera is zero, which is invalid.");
-		return false;
+		UE_LOG(LogCameraCalibrationEditor, Error, TEXT("Could not initiate distortion calibration. The sensor width of the CineCamera is zero, which is invalid."));
+		return CalibrationTask;
 	}
 
 	if (!FocalLengthEstimate.IsSet())
 	{
-		OutErrorMessage = LOCTEXT("NoFocalLengthEstimateSet", "Please enter a value (in mm) for the focal length estimate.");
-		return false;
+		UE_LOG(LogCameraCalibrationEditor, Error, TEXT("Could not initiate distortion calibration. Please enter a value (in mm) for the focal length estimate."));
+		return CalibrationTask;
 	}
 
 	const float FocalLengthEstimateValue = FocalLengthEstimate.GetValue();
@@ -339,10 +331,10 @@ bool UCameraLensDistortionAlgoAruco::GetLensDistortion(
 	FVector2D FocalLength = FVector2D(Fx, Fx);
 	FVector2D ImageCenter = FVector2D((ImageSize.X - 1) * 0.5, (ImageSize.Y - 1) * 0.5);
 
-	TArray<TArray<FVector>> Samples3d;
+	TArray<FObjectPoints> Samples3d;
 	Samples3d.Reserve(CalibrationRows.Num());
 
-	TArray<TArray<FVector2f>> Samples2d;
+	TArray<FImagePoints> Samples2d;
 	Samples2d.Reserve(CalibrationRows.Num());
 
 	TArray<FTransform> CameraPoses;
@@ -353,127 +345,125 @@ bool UCameraLensDistortionAlgoAruco::GetLensDistortion(
 	{
 		check(Row.IsValid());
 
-		TArray<FVector> Points3d;
-		Points3d.Reserve(Row->ArucoPoints.Num() * 4);
+		FObjectPoints Points3d;
+		Points3d.Points.Reserve(Row->ArucoPoints.Num() * 4);
 
-		TArray<FVector2f> Points2d;
-		Points2d.Reserve(Row->ArucoPoints.Num() * 4);
+		FImagePoints Points2d;
+		Points2d.Points.Reserve(Row->ArucoPoints.Num() * 4);
 
 		for (const FArucoCalibrationPoint& Marker : Row->ArucoPoints)
 		{
 			for (int32 CornerIndex = 0; CornerIndex < 4; ++CornerIndex)
 			{
-				Points3d.Add(Marker.Corners3D[CornerIndex]);
-				Points2d.Add(Marker.Corners2D[CornerIndex]);
+				Points3d.Points.Add(Marker.Corners3D[CornerIndex]);
+
+				const FVector2f& MarkerCorner2D = Marker.Corners2D[CornerIndex];
+				Points2d.Points.Add(FVector2D(MarkerCorner2D.X, MarkerCorner2D.Y));
 			}
 		}
 
 		Samples3d.Add(Points3d);
 		Samples2d.Add(Points2d);
-
 		CameraPoses.Add(Row->CameraPose);
 	}
 
-	ECalibrationFlags SolverFlags = ECalibrationFlags::None;
-
-	// Aruco markers may be detected anywhere in the image, and there is no guarantee that they will be coplanar. 
-	// The solver's initialization for focal length assumes all points in an image are coplanar. 
-	// Therefore, we must provide an intrinsics guess to skip this initialization step in the solver. 
-	EnumAddFlags(SolverFlags, ECalibrationFlags::UseIntrinsicGuess);
-
-	if (CVarUseExtrinsicsGuessAruco.GetValueOnGameThread())
-	{
-		EnumAddFlags(SolverFlags, ECalibrationFlags::UseExtrinsicGuess);
-	}
+	const bool bUseNeuralNetSolver = CVarUseNeuralNetSolverAruco.GetValueOnAnyThread();
 
-	if (CVarFixExtrinsicsAruco.GetValueOnAnyThread())
+	if (!bUseNeuralNetSolver)
 	{
-		EnumAddFlags(SolverFlags, ECalibrationFlags::FixExtrinsics);
-	}
+		ECalibrationFlags SolverFlags = ECalibrationFlags::None;
 
-	if (CVarFixZeroDistortionAruco.GetValueOnAnyThread())
-	{
-		EnumAddFlags(SolverFlags, ECalibrationFlags::FixZeroDistortion);
-	}
+		// Aruco markers may be detected anywhere in the image, and there is no guarantee that they will be coplanar. 
+		// The solver's initialization for focal length assumes all points in an image are coplanar. 
+		// Therefore, we must provide an intrinsics guess to skip this initialization step in the solver. 
+		EnumAddFlags(SolverFlags, ECalibrationFlags::UseIntrinsicGuess);
 
-	if (bFixFocalLength)
-	{
-		EnumAddFlags(SolverFlags, ECalibrationFlags::FixFocalLength);
-	}
+		if (CVarUseExtrinsicsGuessAruco.GetValueOnGameThread())
+		{
+			EnumAddFlags(SolverFlags, ECalibrationFlags::UseExtrinsicGuess);
+		}
 
-	if (bFixImageCenter)
-	{
-		EnumAddFlags(SolverFlags, ECalibrationFlags::FixPrincipalPoint);
-	}
+		if (CVarFixExtrinsicsAruco.GetValueOnAnyThread())
+		{
+			EnumAddFlags(SolverFlags, ECalibrationFlags::FixExtrinsics);
+		}
 
-	TArray<float> DistortionCoefficients;
+		if (CVarFixZeroDistortionAruco.GetValueOnAnyThread())
+		{
+			EnumAddFlags(SolverFlags, ECalibrationFlags::FixZeroDistortion);
+		}
 
-	OutError = FCameraCalibrationSolver::CalibrateCamera(
-		LensFile->LensInfo.LensModel,
-		Samples3d,
-		Samples2d,
-		ImageSize,
-		FocalLength,
-		ImageCenter,
-		DistortionCoefficients,
-		CameraPoses,
-		PixelAspect,
-		SolverFlags
-	);
+		if (bFixFocalLength)
+		{
+			EnumAddFlags(SolverFlags, ECalibrationFlags::FixFocalLength);
+		}
 
-	OutLensModel = LensFile->LensInfo.LensModel;
-	OutDistortionInfo.Parameters = DistortionCoefficients;
+		if (bFixImageCenter)
+		{
+			EnumAddFlags(SolverFlags, ECalibrationFlags::FixPrincipalPoint);
+		}
 
-	OutFocalLengthInfo.FxFy = FocalLength / ImageSize;
-	OutImageCenterInfo.PrincipalPoint = ImageCenter / ImageSize;
+		const TSubclassOf<ULensModel> Model = LensFile->LensInfo.LensModel;
 
-	const FLensFileEvaluationInputs LensFileEvalInputs = StepsController->GetLensFileEvaluationInputs();
-	OutFocus = LensFileEvalInputs.Focus;
-	OutZoom = LensFileEvalInputs.Zoom;
+		TArray<float> DistortionCoefficients;
 
-	if (CVarUseExtrinsicsGuessAruco.GetValueOnGameThread())
+		CalibrationTask = UE::Tasks::Launch(UE_SOURCE_LOCATION, [Model, Samples3d, Samples2d, ImageSize, FocalLength, ImageCenter, CameraPoses, PixelAspect, SolverFlags, Focus, Zoom]() mutable
+			{
+				FDistortionCalibrationResult Result;
+
+				Result.ReprojectionError = FCameraCalibrationSolver::CalibrateCamera(
+					Model,
+					Samples3d,
+					Samples2d,
+					ImageSize,
+					FocalLength,
+					ImageCenter,
+					Result.Parameters.Parameters,
+					CameraPoses,
+					PixelAspect,
+					SolverFlags
+				);
+
+				// CalibrateCamera() returns focal length and image center in pixels, but the result is expected to be normalized by the image size
+				Result.FocalLength.FxFy = FocalLength / ImageSize;
+				Result.ImageCenter.PrincipalPoint = ImageCenter / ImageSize;
+
+				// FZ inputs to LUT
+				Result.EvaluatedFocus = Focus;
+				Result.EvaluatedZoom = Zoom;
+
+				return Result;
+			});
+	}
+	else
 	{
-		check(CameraPoses.Num() == CalibrationRows.Num());
+		// ULensDistortionSolver is the base class for the python implementation of the neural network solver.
+		// The intent is for there to be only one derived class, but to be sure, the name is checked to match
+		// the derived class that is expected to be found.
+		TArray<UClass*> DerivedSolverClasses;
+		GetDerivedClasses(ULensDistortionSolver::StaticClass(), DerivedSolverClasses);
 
-		// See if the camera already had an offset applied, in which case we need to account for it.
-		FTransform ExistingOffset = FTransform::Identity;
-
-		bool bWasNodalOffsetApplied = false;
-		if (const ULensComponent* LensComponent = StepsController->FindLensComponent())
+		if (DerivedSolverClasses.Num() == 0)
 		{
-			bWasNodalOffsetApplied = LensComponent->WasNodalOffsetAppliedThisTick();
+			UE_LOG(LogCameraCalibrationEditor, Error, TEXT("Could not initiate distortion calibration. No solver class was found."));
+			return CalibrationTask;
 		}
 
-		if (bWasNodalOffsetApplied)
-		{
-			FNodalPointOffset NodalPointOffset;
+		ULensDistortionSolver* TestSolver = Cast<ULensDistortionSolver>(DerivedSolverClasses.Last()->GetDefaultObject());
 
-			if (LensFile->EvaluateNodalPointOffset(OutFocus, OutZoom, NodalPointOffset))
+		CalibrationTask = UE::Tasks::Launch(UE_SOURCE_LOCATION, [TestSolver, Samples3d, Samples2d, ImageSize, FocalLength, ImageCenter, Focus, Zoom]()
 			{
-				ExistingOffset.SetTranslation(NodalPointOffset.LocationOffset);
-				ExistingOffset.SetRotation(NodalPointOffset.RotationOffset);
-			}
-		}
-
-		TArray<FNodalPointOffset> NewNodalOffsets;
-		NewNodalOffsets.Reserve(CameraPoses.Num());
-		
-		for (int RowIndex = 0; RowIndex < CameraPoses.Num(); ++RowIndex)
-		{
-			const FTransform DesiredCameraTransform = CameraPoses[RowIndex];
-			FTransform DesiredOffset = DesiredCameraTransform * CalibrationRows[RowIndex]->CameraPose.Inverse() * ExistingOffset;
+				FDistortionCalibrationResult Result = TestSolver->Solve(Samples3d, Samples2d, ImageSize, FocalLength, ImageCenter);
 
-			FNodalPointOffset NewNodalOffset;
-			NewNodalOffset.LocationOffset = DesiredOffset.GetLocation();
-			NewNodalOffset.RotationOffset = DesiredOffset.GetRotation();
+				// FZ inputs to LUT
+				Result.EvaluatedFocus = Focus;
+				Result.EvaluatedFocus = Zoom;
 
-			NewNodalOffsets.Add(NewNodalOffset);
-		}
-
-		LensFile->AddNodalOffsetPoint(OutFocus, OutZoom, NewNodalOffsets[0]);
+				return Result;
+			});
 	}
 
-	return true;
+	return CalibrationTask;
 }
 
 void UCameraLensDistortionAlgoAruco::OnDistortionSavedToLens()
diff --git a/Engine/Plugins/VirtualProduction/CameraCalibration/Source/CameraCalibrationEditor/Private/Calibrators/CameraLensDistortionAlgoAruco.h b/Engine/Plugins/VirtualProduction/CameraCalibration/Source/CameraCalibrationEditor/Private/Calibrators/CameraLensDistortionAlgoAruco.h
--- a/Engine/Plugins/VirtualProduction/CameraCalibration/Source/CameraCalibrationEditor/Private/Calibrators/CameraLensDistortionAlgoAruco.h
+++ b/Engine/Plugins/VirtualProduction/CameraCalibration/Source/CameraCalibrationEditor/Private/Calibrators/CameraLensDistortionAlgoAruco.h
@@ -56,16 +56,8 @@ public:
 	virtual FName FriendlyName() const override { return TEXT("Lens Distortion Aruco"); };
 	virtual FName ShortName() const override { return TEXT("Aruco"); };
 	virtual void OnDistortionSavedToLens() override;
-	virtual bool GetLensDistortion(
-		float& OutFocus,
-		float& OutZoom,
-		FDistortionInfo& OutDistortionInfo,
-		FFocalLengthInfo& OutFocalLengthInfo,
-		FImageCenterInfo& OutImageCenterInfo,
-		TSubclassOf<ULensModel>& OutLensModel,
-		double& OutError,
-		FText& OutErrorMessage
-	) override;
+	virtual FDistortionCalibrationTask BeginCalibration() override;
+	virtual bool SupportsAsyncCalibration() override { return true; };
 	virtual bool HasCalibrationData() const override;
 	virtual void PreImportCalibrationData() override;
 	virtual int32 ImportCalibrationRow(const TSharedRef<FJsonObject>& CalibrationRowObject, const FImage& RowImage) override;
diff --git a/Engine/Plugins/VirtualProduction/CameraCalibration/Source/CameraCalibrationEditor/Private/Calibrators/CameraLensDistortionAlgoCheckerboard.cpp b/Engine/Plugins/VirtualProduction/CameraCalibration/Source/CameraCalibrationEditor/Private/Calibrators/CameraLensDistortionAlgoCheckerboard.cpp
--- a/Engine/Plugins/VirtualProduction/CameraCalibration/Source/CameraCalibrationEditor/Private/Calibrators/CameraLensDistortionAlgoCheckerboard.cpp
+++ b/Engine/Plugins/VirtualProduction/CameraCalibration/Source/CameraCalibrationEditor/Private/Calibrators/CameraLensDistortionAlgoCheckerboard.cpp
@@ -41,6 +41,7 @@ static TAutoConsoleVariable<bool> CVarUseIntrinsicsGuess(TEXT("LensDistortionChe
 static TAutoConsoleVariable<bool> CVarFixExtrinsics(TEXT("LensDistortionCheckerboard.FixExtrinsics"), false, TEXT("If true, the solver will fix the camera extrinsics to the user-provided camera poses"));
 static TAutoConsoleVariable<bool> CVarFixZeroDistortion(TEXT("LensDistortionCheckerboard.FixZeroDistortion"), false, TEXT("If true, the solver will fix all distortion values to always be 0"));
 static TAutoConsoleVariable<bool> CVarUseExtrinsicsGuess(TEXT("LensDistortionCheckerboard.UseExtrinsicsGuess"), false, TEXT("If true, the actual checkerboard and camera poses will be used when running the solver"));
+static TAutoConsoleVariable<bool> CVarUseNeuralNetSolver(TEXT("LensDistortionCheckerboard.UseNeuralNetSolver"), false, TEXT("If true, the neural network based solver will solver for distortion and camera intrinsics. Otherwise, the traditional OpenCV solver will be used."));
 #endif
 
 const int UCameraLensDistortionAlgoCheckerboard::DATASET_VERSION = 1;
@@ -609,24 +610,15 @@ bool UCameraLensDistortionAlgoCheckerboard::ValidateNewRow(TSharedPtr<FLensDisto
 	return true;
 }
 
-bool UCameraLensDistortionAlgoCheckerboard::GetLensDistortion(
-	float& OutFocus,
-	float& OutZoom,
-	FDistortionInfo& OutDistortionInfo,
-	FFocalLengthInfo& OutFocalLengthInfo,
-	FImageCenterInfo& OutImageCenterInfo,
-	TSubclassOf<ULensModel>& OutLensModel,
-	double& OutError,
-	FText& OutErrorMessage)
+FDistortionCalibrationTask UCameraLensDistortionAlgoCheckerboard::BeginCalibration()
 {
-	// Sanity checks
-	//
+	FDistortionCalibrationTask CalibrationTask = {};
 
-	// Enough points
+	// Validate that enough points were gathered to attempt a calibration
 	if (CalibrationRows.Num() < 1)
 	{
-		OutErrorMessage = LOCTEXT("NotEnoughSamples", "At least 1 calibration row is required");
-		return false;
+		UE_LOG(LogCameraCalibrationEditor, Error, TEXT("Could not initiate distortion calibration. At least 1 calibration row is required."));
+		return CalibrationTask;
 	}
 
 	// All points are valid
@@ -634,38 +626,31 @@ bool UCameraLensDistortionAlgoCheckerboard::GetLensDistortion(
 	{
 		if (!ensure(Row.IsValid()))
 		{
-			return false;
+			return CalibrationTask;
 		}
 	}
 
-	if (!ensure(Tool.IsValid()))
+	ULensDistortionTool* LensDistortionTool = Tool.Get();
+	if (!LensDistortionTool)
 	{
-		return false;
+		return CalibrationTask;
 	}
 
-	FCameraCalibrationStepsController* StepsController = Tool->GetCameraCalibrationStepsController();
-
-	if (!ensure(StepsController))
+	FCameraCalibrationStepsController* StepsController = LensDistortionTool->GetCameraCalibrationStepsController();
+	if (!StepsController)
 	{
-		return false;
+		return CalibrationTask;
 	}
 
 	ULensFile* LensFile = StepsController->GetLensFile();
-
-	if (!ensure(LensFile))
+	if (!LensFile)
 	{
-		OutErrorMessage = LOCTEXT("LensFileNotFound", "LensFile not found");
-		return false;
+		return CalibrationTask;
 	}
 
-	const TSharedPtr<FLensDistortionCheckerboardRowData>& LastRow = CalibrationRows.Last();
-
-	// Only parameters data mode supported at the moment
-	if (LensFile->DataMode != ELensDataMode::Parameters)
-	{
-		OutErrorMessage = LOCTEXT("OnlyParametersDataModeSupported", "Only Parameters Data Mode supported");
-		return false;
-	}
+	// The evaluation focus and zoom values should not change during data collection, so all rows should have the same values.
+	const float Focus = CalibrationRows[0]->CameraData.InputFocus;
+	const float Zoom = CalibrationRows[0]->CameraData.InputZoom;
 
 	const FIntPoint ImageSize = LensFile->CameraFeedInfo.GetDimensions();
 
@@ -673,8 +658,8 @@ bool UCameraLensDistortionAlgoCheckerboard::GetLensDistortion(
 
 	if (FMath::IsNearlyZero(PixelAspect))
 	{
-		OutErrorMessage = LOCTEXT("PixelAspectNearlyZero", "The pixel aspect in the Lens Information is zero, which is invalid.");
-		return false;
+		UE_LOG(LogCameraCalibrationEditor, Error, TEXT("Could not initiate distortion calibration. The pixel aspect ratio of the CineCamera is zero, which is invalid."));
+		return CalibrationTask;
 	}
 
 	const float PhysicalSensorWidth = StepsController->GetLensFileEvaluationInputs().Filmback.SensorWidth;
@@ -682,123 +667,136 @@ bool UCameraLensDistortionAlgoCheckerboard::GetLensDistortion(
 
 	if (FMath::IsNearlyZero(DesqueezeSensorWidth))
 	{
-		OutErrorMessage = LOCTEXT("SensorWidthNearlyZero", "The sensor width of the CineCamera is zero, which is invalid.");
-		return false;
+		UE_LOG(LogCameraCalibrationEditor, Error, TEXT("Could not initiate distortion calibration. The sensor width of the CineCamera is zero, which is invalid."));
+		return CalibrationTask;
 	}
 
 	const double Fx = (FocalLengthEstimate / DesqueezeSensorWidth) * ImageSize.X;
 
 	// When operating on a desqueezed image, we expect our pixel aspect to be square, so horizontal and vertical field of view are assumed to be equal (i.e. Fx == Fy)
-	FVector2D FocalLength = FVector2D(Fx, Fx); 
+	FVector2D FocalLength = FVector2D(Fx, Fx);
 	FVector2D ImageCenter = FVector2D((ImageSize.X - 1) * 0.5, (ImageSize.Y - 1) * 0.5);
 
-	TArray<TArray<FVector2f>> Samples2d;
-	TArray<TArray<FVector>> Samples3d;
+	TArray<FObjectPoints> Samples3d;
+	TArray<FImagePoints> Samples2d;
+	Samples3d.Reserve(CalibrationRows.Num());
+	Samples2d.Reserve(CalibrationRows.Num());
 
 	TArray<FTransform> CameraPoses;
+	CameraPoses.Reserve(CalibrationRows.Num());
 
 	for (const TSharedPtr<FLensDistortionCheckerboardRowData>& Row : CalibrationRows)
 	{
-		Samples2d.Add(Row->Points2d);
-		Samples3d.Add(Row->Points3d);
-
-		CameraPoses.Add(Row->CameraData.Pose);
-	}
-
-	ECalibrationFlags SolverFlags = ECalibrationFlags::None;
+		FObjectPoints Points3d;
+		Points3d.Points = Row->Points3d;
 
-	if (CVarUseExtrinsicsGuess.GetValueOnGameThread())
-	{
-		EnumAddFlags(SolverFlags, ECalibrationFlags::UseExtrinsicGuess);
-	}
-
-	if (CVarUseIntrinsicsGuess.GetValueOnGameThread())
-	{
-		EnumAddFlags(SolverFlags, ECalibrationFlags::UseIntrinsicGuess);
-	}
-
-	if (CVarFixExtrinsics.GetValueOnGameThread())
-	{
-		EnumAddFlags(SolverFlags, ECalibrationFlags::FixExtrinsics);
-	}
+		FImagePoints Points2d;
+		Points2d.Points.Reserve(Row->Points2d.Num());
+		for (const FVector2f& Point2d : Row->Points2d)
+		{
+			Points2d.Points.Add(FVector2D(Point2d.X, Point2d.Y));
+		}
 
-	if (CVarFixZeroDistortion.GetValueOnGameThread())
-	{
-		EnumAddFlags(SolverFlags, ECalibrationFlags::FixZeroDistortion);
+		Samples3d.Add(Points3d);
+		Samples2d.Add(Points2d);
+		CameraPoses.Add(Row->CameraData.Pose);
 	}
 
-	if (bFixFocalLength)
-	{
-		EnumAddFlags(SolverFlags, ECalibrationFlags::FixFocalLength);
-	}
+	const bool bUseNeuralNetSolver = CVarUseNeuralNetSolver.GetValueOnAnyThread();
 
-	if (bFixImageCenter)
+	if (!bUseNeuralNetSolver)
 	{
-		EnumAddFlags(SolverFlags, ECalibrationFlags::FixPrincipalPoint);
-	}
+		ECalibrationFlags SolverFlags = ECalibrationFlags::None;
 
-	TArray<float> DistortionCoefficients;
+		if (CVarUseExtrinsicsGuess.GetValueOnGameThread())
+		{
+			EnumAddFlags(SolverFlags, ECalibrationFlags::UseExtrinsicGuess);
+		}
 
-	OutError = FCameraCalibrationSolver::CalibrateCamera(
-		LensFile->LensInfo.LensModel,
-		Samples3d,
-		Samples2d,
-		ImageSize,
-		FocalLength,
-		ImageCenter,
-		DistortionCoefficients,
-		CameraPoses,
-		PixelAspect,
-		SolverFlags
-	);
+		if (CVarUseIntrinsicsGuess.GetValueOnGameThread())
+		{
+			EnumAddFlags(SolverFlags, ECalibrationFlags::UseIntrinsicGuess);
+		}
 
-	OutLensModel = LensFile->LensInfo.LensModel;
+		if (CVarFixExtrinsics.GetValueOnGameThread())
+		{
+			EnumAddFlags(SolverFlags, ECalibrationFlags::FixExtrinsics);
+		}
 
-	OutDistortionInfo.Parameters = DistortionCoefficients;
+		if (CVarFixZeroDistortion.GetValueOnGameThread())
+		{
+			EnumAddFlags(SolverFlags, ECalibrationFlags::FixZeroDistortion);
+		}
 
-	OutFocalLengthInfo.FxFy = FocalLength / ImageSize;
+		if (bFixFocalLength)
+		{
+			EnumAddFlags(SolverFlags, ECalibrationFlags::FixFocalLength);
+		}
 
-	OutImageCenterInfo.PrincipalPoint = ImageCenter / ImageSize;
+		if (bFixImageCenter)
+		{
+			EnumAddFlags(SolverFlags, ECalibrationFlags::FixPrincipalPoint);
+		}
 
-	// FZ inputs to LUT
-	OutFocus = LastRow->CameraData.InputFocus;
-	OutZoom = LastRow->CameraData.InputZoom;
+		const TSubclassOf<ULensModel> Model = LensFile->LensInfo.LensModel;
 
-	if (CVarUseExtrinsicsGuess.GetValueOnGameThread())
+		CalibrationTask = UE::Tasks::Launch(UE_SOURCE_LOCATION, [Model, Samples3d, Samples2d, ImageSize, FocalLength, ImageCenter, CameraPoses, PixelAspect, SolverFlags, Focus, Zoom]() mutable
+			{
+				FDistortionCalibrationResult Result;
+
+				Result.ReprojectionError = FCameraCalibrationSolver::CalibrateCamera(
+					Model,
+					Samples3d,
+					Samples2d,
+					ImageSize,
+					FocalLength,
+					ImageCenter,
+					Result.Parameters.Parameters,
+					CameraPoses,
+					PixelAspect,
+					SolverFlags
+				);
+
+				// CalibrateCamera() returns focal length and image center in pixels, but the result is expected to be normalized by the image size
+				Result.FocalLength.FxFy = FocalLength / ImageSize;
+				Result.ImageCenter.PrincipalPoint = ImageCenter / ImageSize;
+
+				// FZ inputs to LUT
+				Result.EvaluatedFocus = Focus;
+				Result.EvaluatedZoom = Zoom;
+
+				return Result;
+			});
+	}
+	else
 	{
-		// See if the camera already had an offset applied, in which case we need to account for it.
-		FTransform ExistingOffset = FTransform::Identity;
+		// ULensDistortionSolver is the base class for the python implementation of the neural network solver.
+		// The intent is for there to be only one derived class, but to be sure, the name is checked to match
+		// the derived class that is expected to be found.
+		TArray<UClass*> DerivedSolverClasses;
+		GetDerivedClasses(ULensDistortionSolver::StaticClass(), DerivedSolverClasses);
 
-		if (CalibrationRows[0]->CameraData.bWasNodalOffsetApplied)
+		if (DerivedSolverClasses.Num() == 0)
 		{
-			FNodalPointOffset NodalPointOffset;
-
-			if (LensFile->EvaluateNodalPointOffset(OutFocus, OutZoom, NodalPointOffset))
-			{
-				ExistingOffset.SetTranslation(NodalPointOffset.LocationOffset);
-				ExistingOffset.SetRotation(NodalPointOffset.RotationOffset);
-			}
+			UE_LOG(LogCameraCalibrationEditor, Error, TEXT("Could not initiate distortion calibration. No solver class was found."));
+			return CalibrationTask;
 		}
 
-		TArray<FNodalPointOffset> NewNodalOffsets;
-		NewNodalOffsets.Reserve(CameraPoses.Num());
+		ULensDistortionSolver* TestSolver = Cast<ULensDistortionSolver>(DerivedSolverClasses.Last()->GetDefaultObject());
 
-		for (int32 RowIndex = 0; RowIndex < CalibrationRows.Num(); ++RowIndex)
-		{
-			const FTransform DesiredCameraTransform = CameraPoses[RowIndex];
-			const FTransform DesiredOffset = DesiredCameraTransform * CalibrationRows[RowIndex]->CameraData.Pose.Inverse() * ExistingOffset;
+		CalibrationTask = UE::Tasks::Launch(UE_SOURCE_LOCATION, [TestSolver, Samples3d, Samples2d, ImageSize, FocalLength, ImageCenter, Focus, Zoom]()
+			{ 
+				FDistortionCalibrationResult Result = TestSolver->Solve(Samples3d, Samples2d, ImageSize, FocalLength, ImageCenter);
 
-			FNodalPointOffset NewNodalOffset;
-			NewNodalOffset.LocationOffset = DesiredOffset.GetLocation();
-			NewNodalOffset.RotationOffset = DesiredOffset.GetRotation();
+				// FZ inputs to LUT
+				Result.EvaluatedFocus = Focus;
+				Result.EvaluatedFocus = Zoom;
 
-			NewNodalOffsets.Add(NewNodalOffset);
-		}
-
-		LensFile->AddNodalOffsetPoint(OutFocus, OutZoom, NewNodalOffsets[0]);
+				return Result;
+			});
 	}
 
-	return true;
+	return CalibrationTask;
 }
 
 TSharedRef<SWidget> UCameraLensDistortionAlgoCheckerboard::BuildCalibrationDevicePickerWidget()
diff --git a/Engine/Plugins/VirtualProduction/CameraCalibration/Source/CameraCalibrationEditor/Private/Calibrators/CameraLensDistortionAlgoCheckerboard.h b/Engine/Plugins/VirtualProduction/CameraCalibration/Source/CameraCalibrationEditor/Private/Calibrators/CameraLensDistortionAlgoCheckerboard.h
--- a/Engine/Plugins/VirtualProduction/CameraCalibration/Source/CameraCalibrationEditor/Private/Calibrators/CameraLensDistortionAlgoCheckerboard.h
+++ b/Engine/Plugins/VirtualProduction/CameraCalibration/Source/CameraCalibrationEditor/Private/Calibrators/CameraLensDistortionAlgoCheckerboard.h
@@ -133,16 +133,8 @@ public:
 	virtual UMaterialInterface* GetOverlayMaterial() const override;
 	virtual bool IsOverlayEnabled() const override { return bShouldShowOverlay; };
 	virtual void OnDistortionSavedToLens() override;
-	virtual bool GetLensDistortion(
-		float& OutFocus,
-		float& OutZoom,
-		FDistortionInfo& OutDistortionInfo,
-		FFocalLengthInfo& OutFocalLengthInfo,
-		FImageCenterInfo& OutImageCenterInfo,
-		TSubclassOf<ULensModel>& OutLensModel,
-		double& OutError,
-		FText& OutErrorMessage
-	) override;
+	virtual FDistortionCalibrationTask BeginCalibration() override;
+	virtual bool SupportsAsyncCalibration() override { return true; };
 	virtual bool HasCalibrationData() const override;
 	virtual void PreImportCalibrationData() override;
 	virtual int32 ImportCalibrationRow(const TSharedRef<FJsonObject>& CalibrationRowObject, const FImage& RowImage) override;
diff --git a/Engine/Plugins/VirtualProduction/CameraCalibration/Source/CameraCalibrationEditor/Private/Calibrators/CameraLensDistortionAlgoPoints.cpp b/Engine/Plugins/VirtualProduction/CameraCalibration/Source/CameraCalibrationEditor/Private/Calibrators/CameraLensDistortionAlgoPoints.cpp
--- a/Engine/Plugins/VirtualProduction/CameraCalibration/Source/CameraCalibrationEditor/Private/Calibrators/CameraLensDistortionAlgoPoints.cpp
+++ b/Engine/Plugins/VirtualProduction/CameraCalibration/Source/CameraCalibrationEditor/Private/Calibrators/CameraLensDistortionAlgoPoints.cpp
@@ -382,28 +382,31 @@ bool UCameraLensDistortionAlgoPoints::GetLensDistortion(
 		return false;
 	}
 
-	TArray<TArray<FVector>> Samples3d;
-	TArray<TArray<FVector2f>> Samples2d;
+	TArray<FObjectPoints> Samples3d;
+	TArray<FImagePoints> Samples2d;
+	Samples3d.Reserve(CalibrationRows.Num());
+	Samples2d.Reserve(CalibrationRows.Num());
 
 	TArray<FTransform> CameraPoses;
 
 	uint32 NextPatternIndex = 0;
 	for (int32 RowIndex = 0; RowIndex < CalibrationRows.Num(); ++RowIndex)
 	{
-		TArray<FVector> Points3d;
-		TArray<FVector2f> Points2d;
-
+		FObjectPoints Points3d;
+		FImagePoints Points2d;
 		while (RowIndex < CalibrationRows.Num() && CalibrationRows[RowIndex]->PatternIndex == NextPatternIndex)
 		{
 			const TSharedPtr<FLensDistortionPointsRowData>& Row = CalibrationRows[RowIndex];
 
-			Points3d.Add(Row->CalibratorPointData.Point3d);
-			Points2d.Add(Row->CalibratorPointData.Point2d);
+			Points3d.Points.Add(Row->CalibratorPointData.Point3d);
+
+			const FVector2f& Point2d = Row->CalibratorPointData.Point2d;
+			Points2d.Points.Add(FVector2D(Point2d.X, Point2d.Y));
 
 			++RowIndex;
 		}
 
-		if (Points3d.Num() > 0)
+		if (Points3d.Points.Num() > 0)
 		{
 			Samples3d.Add(Points3d);
 			Samples2d.Add(Points2d);
@@ -422,10 +425,10 @@ bool UCameraLensDistortionAlgoPoints::GetLensDistortion(
 	}
 
 	// Validate that each pattern has the same number of points
-	const int32 NumPointsInPattern = Samples3d[0].Num();
+	const int32 NumPointsInPattern = Samples3d[0].Points.Num();
 	for (int32 PatternIndex = 1; PatternIndex < Samples3d.Num(); ++PatternIndex)
 	{
-		if (Samples3d[PatternIndex].Num() != NumPointsInPattern)
+		if (Samples3d[PatternIndex].Points.Num() != NumPointsInPattern)
 		{
 			OutErrorMessage = LOCTEXT("DifferentNumPointsInPattern", "Every calibration pattern must have the same number of points");
 			return false;
diff --git a/Engine/Plugins/VirtualProduction/CameraCalibration/Source/CameraCalibrationEditor/Private/TestCameraCalibration.cpp b/Engine/Plugins/VirtualProduction/CameraCalibration/Source/CameraCalibrationEditor/Private/TestCameraCalibration.cpp
--- a/Engine/Plugins/VirtualProduction/CameraCalibration/Source/CameraCalibrationEditor/Private/TestCameraCalibration.cpp
+++ b/Engine/Plugins/VirtualProduction/CameraCalibration/Source/CameraCalibrationEditor/Private/TestCameraCalibration.cpp
@@ -191,12 +191,12 @@ namespace UE::Private::CameraCalibration::AutomatedTests
 
 		// Introduce some random noise to the 3D points. The checkerboard is a rigid object, so the individual 3D positions of each corner cannot change randomly with respect to one another.
 		// However, the entire board could have the wrong pose if, for example, the tracking data is noisy, or if the tracked rigid-body pose sent to UE from the tracking system is not precise.
-		TArray<TArray<FVector>> NoisyObjectPoints;
+		TArray<FObjectPoints> NoisyObjectPoints;
 		NoisyObjectPoints.Reserve(ObjectPoints.Num());
 		for (const TArray<FVector>& Object : ObjectPoints)
 		{
-			TArray<FVector> NoisyPoints;
-			NoisyPoints.Reserve(Object.Num());
+			FObjectPoints NoisyPoints;
+			NoisyPoints.Points.Reserve(Object.Num());
 
 			const double NoiseScale = TestSettings->ObjectPointNoiseScale;
 			const double NoiseX = (FMath::SRand() - 0.5) * (NoiseScale * 2);
@@ -205,7 +205,7 @@ namespace UE::Private::CameraCalibration::AutomatedTests
 
 			for (const FVector& Point : Object)
 			{
-				NoisyPoints.Add(Point + FVector(NoiseX, NoiseY, NoiseZ));
+				NoisyPoints.Points.Add(Point + FVector(NoiseX, NoiseY, NoiseZ));
 			}
 
 			NoisyObjectPoints.Add(NoisyPoints);
@@ -213,12 +213,12 @@ namespace UE::Private::CameraCalibration::AutomatedTests
 
 		// Introduce some random noise to the 2D points. This simulates poor checkerboard detection, which could occur if the checkerboard is not perfectly in-focus, if the image resolution is too low,
 		// or if there is some other imprecision in the corner detection algorithm. 
-		TArray<TArray<FVector2f>> NoisyImagePoints;
+		TArray<FImagePoints> NoisyImagePoints;
 		NoisyImagePoints.Reserve(ImagePoints.Num());
 		for (const TArray<FVector2f>& Image : ImagePoints)
 		{
-			TArray<FVector2f> NoisyPoints;
-			NoisyPoints.Reserve(Image.Num());
+			FImagePoints NoisyPoints;
+			NoisyPoints.Points.Reserve(Image.Num());
 
 			// Unlike the 3D points, the 2D image points could all be randomly noisy compared to one another
 			for (const FVector2f& Point : Image)
@@ -227,7 +227,8 @@ namespace UE::Private::CameraCalibration::AutomatedTests
 				const double NoiseX = (FMath::SRand() - 0.5) * (NoiseScale * 2);
 				const double NoiseY = (FMath::SRand() - 0.5) * (NoiseScale * 2);
 
-				NoisyPoints.Add(Point + FVector2f(NoiseX, NoiseY));
+				const FVector2f PointWithNoise = Point + FVector2f(NoiseX, NoiseY);
+				NoisyPoints.Points.Add(FVector2D(PointWithNoise.X, PointWithNoise.Y));
 			}
 
 			NoisyImagePoints.Add(NoisyPoints);
diff --git a/Engine/Plugins/VirtualProduction/CameraCalibrationCore/Source/CameraCalibrationCore/Public/CameraCalibrationTypes.h b/Engine/Plugins/VirtualProduction/CameraCalibrationCore/Source/CameraCalibrationCore/Public/CameraCalibrationTypes.h
--- a/Engine/Plugins/VirtualProduction/CameraCalibrationCore/Source/CameraCalibrationCore/Public/CameraCalibrationTypes.h
+++ b/Engine/Plugins/VirtualProduction/CameraCalibrationCore/Source/CameraCalibrationCore/Public/CameraCalibrationTypes.h
@@ -8,6 +8,50 @@
 
 #include "CameraCalibrationTypes.generated.h"
 
+/** Results from a distortion calibration, including camera intrinsics and either the parameters to an analytical model or an ST Map */
+USTRUCT(BlueprintType)
+struct FDistortionCalibrationResult
+{
+	GENERATED_BODY()
+		
+	/** Nominal focus distance of the lens associated with this result */
+	UPROPERTY(BlueprintReadWrite, Category = "Distortion")
+	float EvaluatedFocus = 0.0f;
+
+	/** Nominal focal length of the lens associated with this result */
+	UPROPERTY(BlueprintReadWrite, Category = "Distortion")
+	float EvaluatedZoom = 0.0f;
+
+	/** Final reprojection error produced using this result */
+	UPROPERTY(BlueprintReadWrite, Category = "Distortion")
+	double ReprojectionError = 0.0;
+
+	/** Calibrated focal length result */
+	UPROPERTY(BlueprintReadWrite, Category = "Distortion")
+	FFocalLengthInfo FocalLength;
+
+	/** Calibrated image center result */
+	UPROPERTY(BlueprintReadWrite, Category = "Distortion")
+	FImageCenterInfo ImageCenter;
+
+	/** Distortion parameters for the model specified by the lens file. And empty parameter array implies that there is a valid ST Map instead. */
+	UPROPERTY(BlueprintReadWrite, Category = "Distortion")
+	FDistortionInfo Parameters;
+
+	/** ST Map that represents the UV displacements for this result. If the ST Map UTexture is not imported by the solver, a path string should be provided so that the lens distortion tool can import it. */
+	UPROPERTY(BlueprintReadWrite, Category = "Distortion")
+	FSTMapInfo STMap;
+
+	/** Absolute path to an ST Map file on disk that should be imported when this result is processed. */
+	UPROPERTY(BlueprintReadWrite, Category = "Distortion")
+	FString STMapFullPath;
+
+	/** Error text to be written by a solver to provide the reason why the solve may have failed. No error message implies that the solve was successful and the result is valid. */
+	UPROPERTY(BlueprintReadWrite, Category = "Distortion")
+	FText ErrorMessage;
+};
+
+
 /** Utility structure for selecting a distortion handler from the camera calibration subsystem */
 struct UE_DEPRECATED(5.1, "This struct has been deprecated.") FDistortionHandlerPicker;
 
diff --git a/Engine/Plugins/VirtualProduction/CameraCalibrationCore/Source/CameraCalibrationCore/Public/CameraLensDistortionAlgo.h b/Engine/Plugins/VirtualProduction/CameraCalibrationCore/Source/CameraCalibrationCore/Public/CameraLensDistortionAlgo.h
--- a/Engine/Plugins/VirtualProduction/CameraCalibrationCore/Source/CameraCalibrationCore/Public/CameraLensDistortionAlgo.h
+++ b/Engine/Plugins/VirtualProduction/CameraCalibrationCore/Source/CameraCalibrationCore/Public/CameraLensDistortionAlgo.h
@@ -3,6 +3,7 @@
 #pragma once
 
 #include "CoreMinimal.h"
+#include "Tasks/Task.h"
 #include "Widgets/SNullWidget.h"
 #include "Widgets/Text/STextBlock.h"
 
@@ -10,6 +11,7 @@
 
 struct FGeometry;
 struct FPointerEvent;
+struct FDistortionCalibrationResult;
 
 struct FDistortionInfo;
 struct FFocalLengthInfo;
@@ -22,6 +24,9 @@ class ULensModel;
 class UMaterialInterface;
 class SWidget;
 
+
+typedef UE::Tasks::TTask<FDistortionCalibrationResult> FDistortionCalibrationTask;
+
 /**
  * Defines the interface that any lens distortion algorithm should implement
  * in order to be used and listed by the Lens Distortion Tool. 
@@ -80,6 +85,15 @@ public:
 		return false; 
 	};
 
+	/** Launches an asynchronous task to perform the distortion calibration. The task will contain the calibration result when it finishes. */
+	virtual FDistortionCalibrationTask BeginCalibration() { return {}; };
+
+	/** 
+	 * If true, the algo must override BeginCalibration() and use it to launch the calibration as an asynchronous task.
+	 * If false, the algo must override GetLensDistortion() and run the calibration on the Game Thread.
+	 */
+	virtual bool SupportsAsyncCalibration() { return false; };
+
 	/** Returns true if there is any existing calibration data */
 	virtual bool HasCalibrationData() const { return false; };
 
