CameraCalibration: Refactor opencv solver class to inherit from same base as neural net solver. Add an option to the lens distortion tool to pick between the opencv and neural net solvers.
#rnx

diff --git a/Engine/Plugins/VirtualProduction/CameraCalibration/Source/CameraCalibrationEditor/Private/AssetEditor/LensDistortionTool.cpp b/Engine/Plugins/VirtualProduction/CameraCalibration/Source/CameraCalibrationEditor/Private/AssetEditor/LensDistortionTool.cpp
--- a/Engine/Plugins/VirtualProduction/CameraCalibration/Source/CameraCalibrationEditor/Private/AssetEditor/LensDistortionTool.cpp
+++ b/Engine/Plugins/VirtualProduction/CameraCalibration/Source/CameraCalibrationEditor/Private/AssetEditor/LensDistortionTool.cpp
@@ -3,6 +3,7 @@
 #include "LensDistortionTool.h"
 
 #include "AssetToolsModule.h"
+#include "Calibrators/CameraCalibrationSolver.h"
 #include "CameraCalibrationEditorLog.h"
 #include "CameraCalibrationSettings.h"
 #include "CameraCalibrationStepsController.h"
@@ -38,6 +39,15 @@ void ULensDistortionTool::Initialize(TWeakPtr<FCameraCalibrationStepsController>
 {
 	CameraCalibrationStepsController = InCameraCalibrationStepController;
 
+	// Find available solver classes
+
+	TArray<UClass*> DerivedSolverClasses;
+	GetDerivedClasses(ULensDistortionSolver::StaticClass(), DerivedSolverClasses);
+
+	check(!DerivedSolverClasses.IsEmpty());
+
+	SetSolverClass(DerivedSolverClasses[0]);
+
 	// Find available algos
 
 	TArray<TSubclassOf<UCameraLensDistortionAlgo>> Algos;
@@ -250,6 +260,16 @@ void ULensDistortionTool::Deactivate()
 	bIsActive = false;
 }
 
+UClass* ULensDistortionTool::GetSolverClass()
+{
+	return SolverClass;
+}
+
+void ULensDistortionTool::SetSolverClass(UClass* InSolverClass)
+{
+	SolverClass = InSolverClass;
+}
+
 void ULensDistortionTool::ResetAlgo()
 {
 	// Remove old Algo
diff --git a/Engine/Plugins/VirtualProduction/CameraCalibration/Source/CameraCalibrationEditor/Private/AssetEditor/LensDistortionTool.h b/Engine/Plugins/VirtualProduction/CameraCalibration/Source/CameraCalibrationEditor/Private/AssetEditor/LensDistortionTool.h
--- a/Engine/Plugins/VirtualProduction/CameraCalibration/Source/CameraCalibrationEditor/Private/AssetEditor/LensDistortionTool.h
+++ b/Engine/Plugins/VirtualProduction/CameraCalibration/Source/CameraCalibrationEditor/Private/AssetEditor/LensDistortionTool.h
@@ -60,6 +60,12 @@ public:
 	virtual bool IsOverlayEnabled() const override;
 	//~ End UCameraCalibrationStep interface
 
+	/** Get the currently selected lens distortion solver class */
+	UClass* GetSolverClass();
+
+	/** Set the lens distortion solver class to use */
+	void SetSolverClass(UClass* InSolverClass);
+
 	/** Selects the algorithm by name */
 	void SetAlgo(const FName& AlgoName);
 
@@ -134,6 +140,9 @@ private:
 	UPROPERTY(Transient)
 	TObjectPtr<UCameraLensDistortionAlgo> CurrentAlgo;
 
+	/** The solver class that will be used by the algos when creating a solver to calibrate for distortion */
+	TObjectPtr<UClass> SolverClass;
+
 	/** Holds the registered camera lens distortion algos */
 	UPROPERTY(Transient)
 	TMap<FName, TSubclassOf<UCameraLensDistortionAlgo>> AlgosMap;
diff --git a/Engine/Plugins/VirtualProduction/CameraCalibration/Source/CameraCalibrationEditor/Private/AssetEditor/SLensDistortionToolPanel.cpp b/Engine/Plugins/VirtualProduction/CameraCalibration/Source/CameraCalibrationEditor/Private/AssetEditor/SLensDistortionToolPanel.cpp
--- a/Engine/Plugins/VirtualProduction/CameraCalibration/Source/CameraCalibrationEditor/Private/AssetEditor/SLensDistortionToolPanel.cpp
+++ b/Engine/Plugins/VirtualProduction/CameraCalibration/Source/CameraCalibrationEditor/Private/AssetEditor/SLensDistortionToolPanel.cpp
@@ -3,15 +3,17 @@
 #include "SLensDistortionToolPanel.h"
 
 #include "AssetRegistry/AssetData.h"
+#include "Calibrators/CameraCalibrationSolver.h"
 #include "CameraCalibrationSettings.h"
 #include "CameraCalibrationSubsystem.h"
 #include "CameraLensDistortionAlgo.h"
 #include "Dialog/SCustomDialog.h"
 #include "EditorFontGlyphs.h"
-#include "Styling/AppStyle.h"
 #include "Engine/Selection.h"
+#include "Framework/MultiBox/MultiBoxBuilder.h"
 #include "LensDistortionTool.h"
 #include "LensFile.h"
+#include "Styling/AppStyle.h"
 #include "Styling/CoreStyle.h"
 #include "Styling/SlateStyle.h"
 #include "UI/CameraCalibrationWidgetHelpers.h"
@@ -41,6 +43,10 @@ void SLensDistortionToolPanel::Construct(const FArguments& InArgs, ULensDistorti
 		[ 
 			SNew(SVerticalBox)
 
+			+ SVerticalBox::Slot() // Solver picker
+			.MaxHeight(FCameraCalibrationWidgetHelpers::DefaultRowHeight)
+			[FCameraCalibrationWidgetHelpers::BuildLabelWidgetPair(LOCTEXT("SolverPickerWidgetText", "Lens Distortion Solver"), BuildSolverPickerWidget())]
+
 			+ SVerticalBox::Slot() // Algo picker
 			.MaxHeight(FCameraCalibrationWidgetHelpers::DefaultRowHeight)
 			[ FCameraCalibrationWidgetHelpers::BuildLabelWidgetPair(LOCTEXT("LensDistortionAlgo", "Lens Distortion Algo"), BuildAlgoPickerWidget())]
@@ -197,6 +203,81 @@ void SLensDistortionToolPanel::UpdateAlgosOptions()
 	AlgosComboBox->RefreshOptions();
 }
 
+TSharedRef<SWidget> SLensDistortionToolPanel::BuildSolverPickerWidget()
+{
+	return SNew(SHorizontalBox)
+
+		+ SHorizontalBox::Slot()
+		.AutoWidth()
+		.Padding(FMargin(0.0f, 0.0f, 0.0f, 0.0f))
+		[
+			SNew(SComboButton)
+			.OnGetMenuContent(this, &SLensDistortionToolPanel::BuildSolverList)
+			.ContentPadding(FMargin(4.0, 2.0))
+			.ButtonContent()
+			[
+				SNew(STextBlock).Text(this, &SLensDistortionToolPanel::GetSelectedSolverName)
+			]
+		];
+}
+
+TSharedRef<SWidget> SLensDistortionToolPanel::BuildSolverList() 
+{
+	// Generate menu
+	FMenuBuilder MenuBuilder(true, nullptr);
+	MenuBuilder.BeginSection("DistortionSolverSection", LOCTEXT("DistortionSolverSectionText", "Distortion Solvers"));
+
+	TArray<UClass*> DerivedSolverClasses;
+	GetDerivedClasses(ULensDistortionSolver::StaticClass(), DerivedSolverClasses);
+
+	if (!DerivedSolverClasses.IsEmpty())
+	{
+		for (UClass* SolverClass : DerivedSolverClasses)
+		{
+			ULensDistortionSolver* SolverCDO = Cast<ULensDistortionSolver>(SolverClass->GetDefaultObject());
+
+			MenuBuilder.AddMenuEntry(
+				SolverCDO->GetDisplayName(),
+				SolverCDO->GetDisplayName(), // TODO: If this remains a user-facing option, then it may be better to have the solver implement a class description text field to use as the tooltip instead
+				FSlateIcon(),
+				FUIAction(
+					FExecuteAction::CreateSP(this, &SLensDistortionToolPanel::OnSolverClassSelected, SolverClass),
+					FCanExecuteAction(),
+					FIsActionChecked::CreateSP(this, &SLensDistortionToolPanel::IsSolverClassSelected, SolverClass)
+				),
+				NAME_None,
+				EUserInterfaceActionType::RadioButton
+			);
+		}
+	}
+
+	MenuBuilder.EndSection();
+
+	return MenuBuilder.MakeWidget();
+
+}
+
+void SLensDistortionToolPanel::OnSolverClassSelected(UClass* SolverClass)
+{
+	Tool->SetSolverClass(SolverClass);
+}
+
+bool SLensDistortionToolPanel::IsSolverClassSelected(UClass* SolverClass) const
+{
+	return (Tool->GetSolverClass() == SolverClass);
+}
+
+FText SLensDistortionToolPanel::GetSelectedSolverName() const
+{
+	if (UClass* SolverClass = Tool->GetSolverClass())
+	{
+		ULensDistortionSolver* SolverCDO = Cast<ULensDistortionSolver>(SolverClass->GetDefaultObject());
+		return SolverCDO->GetDisplayName();
+	}
+	
+	return FText::GetEmpty();
+}
+
 TSharedRef<SWidget> SLensDistortionToolPanel::BuildAlgoPickerWidget()
 {
 	// Create ComboBox widget
diff --git a/Engine/Plugins/VirtualProduction/CameraCalibration/Source/CameraCalibrationEditor/Private/AssetEditor/SLensDistortionToolPanel.h b/Engine/Plugins/VirtualProduction/CameraCalibration/Source/CameraCalibrationEditor/Private/AssetEditor/SLensDistortionToolPanel.h
--- a/Engine/Plugins/VirtualProduction/CameraCalibration/Source/CameraCalibrationEditor/Private/AssetEditor/SLensDistortionToolPanel.h
+++ b/Engine/Plugins/VirtualProduction/CameraCalibration/Source/CameraCalibrationEditor/Private/AssetEditor/SLensDistortionToolPanel.h
@@ -35,6 +35,21 @@ private:
 	/** Builds the wrapper for the currently selected algo UI */
 	TSharedRef<SWidget> BuildUIWrapper();
 
+	/** Builds the UI of the solver picker widget */
+	TSharedRef<SWidget> BuildSolverPickerWidget();
+
+	/** Builds the list of solvers to display when opening the solver picker */
+	TSharedRef<SWidget> BuildSolverList();
+
+	/** When the solver class is selected in the UI, update the lens distortion tool with the newly selected class */
+	void OnSolverClassSelected(UClass* SolverClass);
+
+	/** Returns true if the input class is the currently selected one */
+	bool IsSolverClassSelected(UClass* SolverClass) const;
+
+	/** Get the display name of the currently selected solver class */
+	FText GetSelectedSolverName() const;
+
 	/** Builds the UI for the algorithm picker */
 	TSharedRef<SWidget> BuildAlgoPickerWidget();
 
diff --git a/Engine/Plugins/VirtualProduction/CameraCalibration/Source/CameraCalibrationEditor/Private/Calibrators/CameraCalibrationSolver.cpp b/Engine/Plugins/VirtualProduction/CameraCalibration/Source/CameraCalibrationEditor/Private/Calibrators/CameraCalibrationSolver.cpp
--- a/Engine/Plugins/VirtualProduction/CameraCalibration/Source/CameraCalibrationEditor/Private/Calibrators/CameraCalibrationSolver.cpp
+++ b/Engine/Plugins/VirtualProduction/CameraCalibration/Source/CameraCalibrationEditor/Private/Calibrators/CameraCalibrationSolver.cpp
@@ -15,26 +15,34 @@
 
 #endif	// WITH_OPENCV
 
+#define LOCTEXT_NAMESPACE "CameraCalibrationSolver"
+
 static TAutoConsoleVariable<int> CVarUseLegacySphericalSolver(TEXT("CameraCalibration.UseLegacySphericalSolver"), 0, TEXT("If set, the legacy OpenCV spherical solver will be used"));
 static TAutoConsoleVariable<float> CVarRotationStepValue(TEXT("CameraCalibration.RotationStepValue"), 0.05, TEXT("The value of the initial step size to use when finding an optimal nodal offset rotation that minimizes reprojection error."));
 static TAutoConsoleVariable<float> CVarLocationStepValue(TEXT("CameraCalibration.LocationStepValue"), 0.5, TEXT("The value of the initial step size to use when finding an optimal nodal offset location that minimizes reprojection error."));
 
 DEFINE_LOG_CATEGORY_STATIC(LogCameraCalibrationSolver, Log, All);
 
-double FCameraCalibrationSolver::CalibrateCamera(
-	const TSubclassOf<ULensModel> LensModel,
-	const TArray<FObjectPoints>& InObjectPointsArray,
-	const TArray<FImagePoints>& InImagePointsArray,
+FText ULensDistortionSolverOpenCV::GetDisplayName_Implementation() const
+{
+	return LOCTEXT("OpenCVSolverDisplayName", "OpenCV Solver");
+}
+
+FDistortionCalibrationResult ULensDistortionSolverOpenCV::Solve_Implementation(
+	const TArray<FObjectPoints>& ObjectPointArray,
+	const TArray<FImagePoints>& ImagePointArray,
 	const FIntPoint ImageSize,
-	FVector2D& InOutFocalLength,
-	FVector2D& InOutImageCenter,
-	TArray<float>& OutDistCoeffs,
-	TArray<FTransform>& InOutCameraPoses,
+	const FVector2D& FocalLength,
+	const FVector2D& ImageCenter,
+	const TArray<FTransform>& CameraPoses,
+	TSubclassOf<ULensModel> LensModel,
 	double PixelAspect,
 	ECalibrationFlags SolverFlags)
 {
+	FDistortionCalibrationResult Result;
+
 #if WITH_OPENCV
-	const int NumImages = InObjectPointsArray.Num();
+	const int NumImages = ObjectPointArray.Num();
 
 	// Create an array to store the number of points in each image
 	cv::Mat NumPointsMat = cv::Mat(1, NumImages, CV_32S);
@@ -43,7 +51,7 @@ double FCameraCalibrationSolver::CalibrateCamera(
 	int MaxPoints = 0;
 	for (int ImageIndex = 0; ImageIndex < NumImages; ++ImageIndex)
 	{
-		const int NumPointsInImage = InObjectPointsArray[ImageIndex].Points.Num();
+		const int NumPointsInImage = ObjectPointArray[ImageIndex].Points.Num();
 
 		NumPointsMat.at<int>(ImageIndex) = NumPointsInImage;
 		NumTotalPoints += NumPointsInImage;
@@ -55,18 +63,16 @@ double FCameraCalibrationSolver::CalibrateCamera(
 	cv::Mat ImagePointsMat = cv::Mat(1, NumTotalPoints, CV_64FC2);
 
 	// Reorganize the 3D and 2D points from the input arrays or arrays to be laid out linearly in memory in two cv::Mat objects
-	GatherPoints(InObjectPointsArray, InImagePointsArray, ObjectPointsMat, ImagePointsMat);
+	GatherPoints(ObjectPointArray, ImagePointArray, ObjectPointsMat, ImagePointsMat);
 
 	double RMSE = 0.0;
 
 	cv::Mat CameraMatrix = cv::Mat::eye(3, 3, CV_64F);
 
-	CameraMatrix.at<double>(0, 0) = InOutFocalLength.X;
-	CameraMatrix.at<double>(1, 1) = InOutFocalLength.Y;
-	CameraMatrix.at<double>(0, 2) = InOutImageCenter.X;
-	CameraMatrix.at<double>(1, 2) = InOutImageCenter.Y;
-
-	OutDistCoeffs.Empty();
+	CameraMatrix.at<double>(0, 0) = FocalLength.X;
+	CameraMatrix.at<double>(1, 1) = FocalLength.Y;
+	CameraMatrix.at<double>(0, 2) = ImageCenter.X;
+	CameraMatrix.at<double>(1, 2) = ImageCenter.Y;
 
 	const int NumDistortionCoefficients = LensModel->GetDefaultObject<ULensModel>()->GetNumParameters();
 	cv::Mat DistCoeffs = cv::Mat(1, NumDistortionCoefficients, CV_64F);
@@ -96,10 +102,10 @@ double FCameraCalibrationSolver::CalibrateCamera(
 		std::vector<std::vector<cv::Point2f>> Samples2d;
 		std::vector<std::vector<cv::Point3f>> Samples3d;
 
-		Samples2d.reserve(InImagePointsArray.Num());
-		Samples3d.reserve(InObjectPointsArray.Num());
+		Samples2d.reserve(ImagePointArray.Num());
+		Samples3d.reserve(ObjectPointArray.Num());
 
-		for (const FObjectPoints& PointsInImage : InObjectPointsArray)
+		for (const FObjectPoints& PointsInImage : ObjectPointArray)
 		{
 			std::vector<cv::Point3f> Points3d;
 			Points3d.reserve(PointsInImage.Points.Num());
@@ -113,7 +119,7 @@ double FCameraCalibrationSolver::CalibrateCamera(
 			Samples3d.push_back(Points3d);
 		}
 
-		for (const FImagePoints& PointsInImage : InImagePointsArray)
+		for (const FImagePoints& PointsInImage : ImagePointArray)
 		{
 			std::vector<cv::Point2f> Points2d;
 			Points2d.reserve(PointsInImage.Points.Num());
@@ -155,19 +161,21 @@ double FCameraCalibrationSolver::CalibrateCamera(
 		);
 
 		// Set the output intrinsics and distortion parameters to the final values calculated by the solver
-		InOutFocalLength.X = CameraMatrix.at<double>(0, 0);
-		InOutFocalLength.Y = CameraMatrix.at<double>(1, 1);
-		InOutImageCenter.X = CameraMatrix.at<double>(0, 2);
-		InOutImageCenter.Y = CameraMatrix.at<double>(1, 2);
+		Result.FocalLength.FxFy.X = CameraMatrix.at<double>(0, 0);
+		Result.FocalLength.FxFy.Y = CameraMatrix.at<double>(1, 1);
+		Result.ImageCenter.PrincipalPoint.X = CameraMatrix.at<double>(0, 2);
+		Result.ImageCenter.PrincipalPoint.Y = CameraMatrix.at<double>(1, 2);
 
 		// The spherical distortion coefficients in our model are in a different order than they appear in the solver, so the results need to be rearranged
-		OutDistCoeffs.Add(DistCoeffs.at<double>(0));
-		OutDistCoeffs.Add(DistCoeffs.at<double>(1));
-		OutDistCoeffs.Add(DistCoeffs.at<double>(4));
-		OutDistCoeffs.Add(DistCoeffs.at<double>(2));
-		OutDistCoeffs.Add(DistCoeffs.at<double>(3));
+		Result.Parameters.Parameters.Add(DistCoeffs.at<double>(0));
+		Result.Parameters.Parameters.Add(DistCoeffs.at<double>(1));
+		Result.Parameters.Parameters.Add(DistCoeffs.at<double>(4));
+		Result.Parameters.Parameters.Add(DistCoeffs.at<double>(2));
+		Result.Parameters.Parameters.Add(DistCoeffs.at<double>(3));
 
-		return RMSE;
+		Result.ReprojectionError = RMSE;
+
+		return Result;
 	}
 
 	// If the flag to use a starting guess for the camera intrinsics is not set, calculate some initial values for the intrinsic parameters
@@ -248,7 +256,7 @@ double FCameraCalibrationSolver::CalibrateCamera(
 		cv::Mat Rotation = Solver.Params.rowRange(NumIntrinsics, NumIntrinsics + 3);
 		cv::Mat Translation = Solver.Params.rowRange(NumIntrinsics + 3, NumIntrinsics + 6);
 
-		FOpenCVHelper::MakeObjectVectorsFromCameraPose(InOutCameraPoses[0], Rotation, Translation);
+		FOpenCVHelper::MakeObjectVectorsFromCameraPose(CameraPoses[0], Rotation, Translation);
 	}
 	else
 	{
@@ -278,11 +286,11 @@ double FCameraCalibrationSolver::CalibrateCamera(
 	TArray<FTransform> CameraMovements;
 	if ((EnumHasAnyFlags(SolverFlags, ECalibrationFlags::UseExtrinsicGuess)))
 	{
-		CameraMovements.Reserve(InOutCameraPoses.Num());
+		CameraMovements.Reserve(CameraPoses.Num());
 
-		for (const FTransform& Pose : InOutCameraPoses)
+		for (const FTransform& Pose : CameraPoses)
 		{
-			const FTransform CameraMovement = InOutCameraPoses[0].Inverse() * Pose;
+			const FTransform CameraMovement = CameraPoses[0].Inverse() * Pose;
 			CameraMovements.Add(CameraMovement);
 		}
 	}
@@ -404,40 +412,34 @@ double FCameraCalibrationSolver::CalibrateCamera(
 	RMSE = FMath::Sqrt(ReprojectionError / NumTotalPoints);
 
 	// Set the output intrinsics and distortion parameters to the final values calculated by the solver
-	InOutFocalLength.X = CameraMatrix.at<double>(0, 0);
-	InOutFocalLength.Y = CameraMatrix.at<double>(1, 1);
-	InOutImageCenter.X = CameraMatrix.at<double>(0, 2);
-	InOutImageCenter.Y = CameraMatrix.at<double>(1, 2);
+	Result.FocalLength.FxFy.X = CameraMatrix.at<double>(0, 0);
+	Result.FocalLength.FxFy.Y = CameraMatrix.at<double>(1, 1);
+	Result.ImageCenter.PrincipalPoint.X = CameraMatrix.at<double>(0, 2);
+	Result.ImageCenter.PrincipalPoint.Y = CameraMatrix.at<double>(1, 2);
 
 	if (LensModel == UAnamorphicLensModel::StaticClass())
 	{
 		for (int CoeffIndex = 0; CoeffIndex < NumDistortionCoefficients; ++CoeffIndex)
 		{
-			OutDistCoeffs.Add(DistCoeffs.at<double>(CoeffIndex));
+			Result.Parameters.Parameters.Add(DistCoeffs.at<double>(CoeffIndex));
 		}
 	}
 	else
 	{
 		// The spherical distortion coefficients in our model are in a different order than they appear in the solver, so the results need to be rearranged
-		OutDistCoeffs.Add(DistCoeffs.at<double>(0));
-		OutDistCoeffs.Add(DistCoeffs.at<double>(1));
-		OutDistCoeffs.Add(DistCoeffs.at<double>(4));
-		OutDistCoeffs.Add(DistCoeffs.at<double>(2));
-		OutDistCoeffs.Add(DistCoeffs.at<double>(3));
+		Result.Parameters.Parameters.Add(DistCoeffs.at<double>(0));
+		Result.Parameters.Parameters.Add(DistCoeffs.at<double>(1));
+		Result.Parameters.Parameters.Add(DistCoeffs.at<double>(4));
+		Result.Parameters.Parameters.Add(DistCoeffs.at<double>(2));
+		Result.Parameters.Parameters.Add(DistCoeffs.at<double>(3));
 	}
 
-	for (int ImageIndex = 0; ImageIndex < NumPosesToSolve; ImageIndex++)
-	{
-		const int ExtrinsicOffset = NumIntrinsics + (ImageIndex * NumExtrinsics);
-		cv::Mat Rotation = Solver.Params.rowRange(ExtrinsicOffset, ExtrinsicOffset + 3);
-		cv::Mat Translation = Solver.Params.rowRange(ExtrinsicOffset + 3, ExtrinsicOffset + 6);
-
-		FOpenCVHelper::MakeCameraPoseFromObjectVectors(Rotation, Translation, InOutCameraPoses[ImageIndex]);
-	}
+	Result.ReprojectionError = RMSE;
 
-	return RMSE;
+	return Result;
 #else
-	return -1.0;
+	Result.ErrorMessage = LOCTEXT("OpenCVNotSupportedError", "OpenCV is not supported");
+	return Result;
 #endif // WITH_OPENCV
 }
 
@@ -518,7 +520,7 @@ double FCameraCalibrationSolver::OptimizeNodalOffset(
 }
 
 #if WITH_OPENCV
-void FCameraCalibrationSolver::InitCameraIntrinsics(
+void ULensDistortionSolverOpenCV::InitCameraIntrinsics(
 	const cv::Mat& ObjectPoints,
 	const cv::Mat& ImagePoints,
 	const cv::Mat& NumPoints,
@@ -618,7 +620,7 @@ void FCameraCalibrationSolver::InitCameraIntrinsics(
 	CameraMatrix.at<double>(1, 2) = Cy;
 }
 
-void FCameraCalibrationSolver::InitCameraExtrinsics(
+void ULensDistortionSolverOpenCV::InitCameraExtrinsics(
 	const TSubclassOf<ULensModel> LensModel,
 	const cv::Mat& ObjectPoints,
 	const cv::Mat& ImagePoints,
@@ -776,7 +778,7 @@ void FCameraCalibrationSolver::InitCameraExtrinsics(
 	}
 }
 
-void FCameraCalibrationSolver::ProjectPoints(
+void ULensDistortionSolverOpenCV::ProjectPoints(
 	const TSubclassOf<ULensModel> LensModel,
 	const cv::Mat& ObjectPoints,
 	const cv::Mat& Rotation,
@@ -791,7 +793,7 @@ void FCameraCalibrationSolver::ProjectPoints(
 	ProjectPoints(LensModel, ObjectPoints, Rotation, Translation, CameraMatrix, DistCoeffs, ImageSize, ProjectedPoints, Jacobian, SolverFlags);
 }
 
-void FCameraCalibrationSolver::ProjectPoints(
+void ULensDistortionSolverOpenCV::ProjectPoints(
 	const TSubclassOf<ULensModel> LensModel,
 	const cv::Mat& ObjectPoints,
 	const cv::Mat& Rotation,
@@ -831,7 +833,7 @@ void FCameraCalibrationSolver::ProjectPoints(
 	}
 }
 
-void FCameraCalibrationSolver::ProjectPointsAnamorphic(
+void ULensDistortionSolverOpenCV::ProjectPointsAnamorphic(
 	const cv::Mat& ObjectPoints,
 	const cv::Mat& Rotation,
 	const cv::Mat& Translation,
@@ -1154,7 +1156,7 @@ void FCameraCalibrationSolver::ProjectPointsAnamorphic(
 	}
 }
 
-void FCameraCalibrationSolver::ProjectPointsSpherical(
+void ULensDistortionSolverOpenCV::ProjectPointsSpherical(
 	const cv::Mat& ObjectPoints,
 	const cv::Mat& Rotation,
 	const cv::Mat& Translation,
@@ -1369,7 +1371,7 @@ void FCameraCalibrationSolver::ProjectPointsSpherical(
 	}
 }
 
-void FCameraCalibrationSolver::GatherPoints(
+void ULensDistortionSolverOpenCV::GatherPoints(
 	const TArray<FObjectPoints>& InObjectPointsArray,
 	const TArray<FImagePoints>& InImagePointsArray,
 	cv::Mat& ObjectPointsMat,
@@ -1401,7 +1403,7 @@ void FCameraCalibrationSolver::GatherPoints(
 	}
 }
 
-void FCameraCalibrationSolver::SubdivideJacobian(
+void ULensDistortionSolverOpenCV::SubdivideJacobian(
 	const cv::Mat& Jacobian,
 	cv::Mat& JacRotation,
 	cv::Mat& JacTranslation,
@@ -1674,3 +1676,5 @@ double FOptimizeNodalOffsetSolver::calc(const double* x) const
 	return RootMeanSquareError;
 }
 #endif // WITH_OPENCV
+
+#undef LOCTEXT_NAMESPACE
diff --git a/Engine/Plugins/VirtualProduction/CameraCalibration/Source/CameraCalibrationEditor/Private/Calibrators/CameraCalibrationSolver.h b/Engine/Plugins/VirtualProduction/CameraCalibration/Source/CameraCalibrationEditor/Private/Calibrators/CameraCalibrationSolver.h
--- a/Engine/Plugins/VirtualProduction/CameraCalibration/Source/CameraCalibrationEditor/Private/Calibrators/CameraCalibrationSolver.h
+++ b/Engine/Plugins/VirtualProduction/CameraCalibration/Source/CameraCalibrationEditor/Private/Calibrators/CameraCalibrationSolver.h
@@ -2,6 +2,7 @@
 
 #pragma once
 
+#include "CameraCalibrationTypes.h"
 #include "Models/LensModel.h"
 #include "Templates/SubclassOf.h"
 
@@ -43,75 +44,93 @@ struct FImagePoints
 	TArray<FVector2D> Points;
 };
 
+/**
+  * Flags used to modify the execution of the calibration solver
+  */
+UENUM()
+enum class ECalibrationFlags : uint32
+{
+	None = 0,
+	UseIntrinsicGuess = 1 << 0,  /** The solver will skip intrinsics initialization an use the input values of Focal Length and Image Center as the starting values for the optimization */
+	UseExtrinsicGuess = 1 << 1,  /** The solver will skip extrinsics initialization an use the input Camera Poses as the starting value for the optimization */
+	FixFocalLength = 1 << 2,     /** The solver will not optimize the focal length */
+	FixPrincipalPoint = 1 << 3,  /** The solver will not optimize the principal point */
+	FixExtrinsics = 1 << 4,      /** The solver will not optimize the camera extrinsics */
+	FixZeroDistortion = 1 << 5   /** The solver will fix all distortion values at 0 */
+};
+
 /** 
  * Base lens distortion solver class that can be inherited from in Blueprints or Python 
  */
-UCLASS(Blueprintable)
+UCLASS(Abstract, Blueprintable)
 class ULensDistortionSolver : public UObject
 {
 	GENERATED_BODY()
 
 public:
 	/** Calibrate camera intrinsics and distortion from a set of input 3D-2D point correspondences and initial camera intrinsics guess. */
-	UFUNCTION(BlueprintImplementableEvent, Category = "Calibration")
+	UFUNCTION(BlueprintNativeEvent, Category = "Calibration")
 	FDistortionCalibrationResult Solve(
 		const TArray<FObjectPoints>& ObjectPointArray,
 		const TArray<FImagePoints>& ImagePointArray,
 		const FIntPoint ImageSize,
 		const FVector2D& FocalLength,
-		const FVector2D& ImageCenter) const;
-};
+		const FVector2D& ImageCenter,
+		const TArray<FTransform>& CameraPoses,
+		TSubclassOf<ULensModel> LensModel,
+		double PixelAspect,
+		ECalibrationFlags SolverFlags);
 
-/**
-  * Flags used to modify the execution of the calibration solver
-  */
-enum class ECalibrationFlags : uint32
-{
-	None = 0,
-	UseIntrinsicGuess = 1 << 0,  /** The solver will skip intrinsics initialization an use the input values of Focal Length and Image Center as the starting values for the optimization */
-	UseExtrinsicGuess = 1 << 1,  /** The solver will skip extrinsics initialization an use the input Camera Poses as the starting value for the optimization */
-	FixFocalLength = 1 << 2,     /** The solver will not optimize the focal length */
-	FixPrincipalPoint = 1 << 3,  /** The solver will not optimize the principal point */
-	FixExtrinsics = 1 << 4,      /** The solver will not optimize the camera extrinsics */
-	FixZeroDistortion = 1 << 5   /** The solver will fix all distortion values at 0 */
+	/** Get the name of this solver class for displaying in the editor UI */
+	UFUNCTION(BlueprintNativeEvent, Category = "Calibration")
+	FText GetDisplayName() const;
+
+	virtual FDistortionCalibrationResult Solve_Implementation(
+		const TArray<FObjectPoints>& ObjectPointArray,
+		const TArray<FImagePoints>& ImagePointArray,
+		const FIntPoint ImageSize,
+		const FVector2D& FocalLength,
+		const FVector2D& ImageCenter,
+		const TArray<FTransform>& CameraPoses,
+		TSubclassOf<ULensModel> LensModel,
+		double PixelAspect,
+		ECalibrationFlags SolverFlags) PURE_VIRTUAL(ULensDistortionSolver::Solve_Implementation, return FDistortionCalibrationResult(););
+
+	virtual FText GetDisplayName_Implementation() const PURE_VIRTUAL(ULensDistortionSolver::GetDisplayName_Implementation, return FText::GetEmpty(););
 };
 
 /** 
   * Lens Distortion Solver class, supporting anamorphic and spherical models 
   * The implementation is largely based on the implementation of calibrateCamera from OpenCV: [REDACTED_PROFILE_URL]/opencv
   */
-class FCameraCalibrationSolver
+
+UCLASS()
+class ULensDistortionSolverOpenCV : public ULensDistortionSolver
 {
+	GENERATED_BODY()
+
 public:
 	/** 
 	  * Calibrate camera intrinsics and distortion parameters using a set of input 3D and 2D point correspondences
 	  * Returns the root mean reprojection error between the input image points and the projection of the input object points
 	  */
-	static double CalibrateCamera(
-		const TSubclassOf<ULensModel> LensModel,
-		const TArray<FObjectPoints>& InObjectPointsArray,
-		const TArray<FImagePoints>& InImagePointsArray,
+	virtual FDistortionCalibrationResult Solve_Implementation(
+		const TArray<FObjectPoints>& ObjectPointArray,
+		const TArray<FImagePoints>& ImagePointArray,
 		const FIntPoint ImageSize,
-		FVector2D& InOutFocalLength,
-		FVector2D& InOutImageCenter,
-		TArray<float>& OutDistCoeffs,
-		TArray<FTransform>& InOutCameraPoses,
-		double PixelAspect = 1.0,
-		ECalibrationFlags SolverFlags = ECalibrationFlags::None);
+		const FVector2D& FocalLength,
+		const FVector2D& ImageCenter,
+		const TArray<FTransform>& CameraPoses,
+		TSubclassOf<ULensModel> LensModel,
+		double PixelAspect,
+		ECalibrationFlags SolverFlags) override;
 
-	/** Optimize the input nodal offset transform by running a downhill solver that attempts to minimize the reprojection error of the input points */
-	static double OptimizeNodalOffset(
-		const TArray<TArray<FVector>>& InObjectPoints, 
-		const TArray<TArray<FVector2f>>& InImagePoints, 
-		const FVector2D& InFocalLength,
-		const FVector2D& InImageCenter,
-		const TArray<FTransform>& InCameraPoses,
-		FTransform& InOutNodalOffset);
+	virtual FText GetDisplayName_Implementation() const override;
 
 #if WITH_OPENCV
 private:
 	/** Initialize the camera matrix of intrinsic parameters using linear algebra techniques */
-	static void InitCameraIntrinsics(
+	void InitCameraIntrinsics(
 		const cv::Mat& ObjectPoints,
 		const cv::Mat& ImagePoints,
 		const cv::Mat& NumPoints,
@@ -119,7 +138,7 @@ private:
 		cv::Mat& CameraMatrix);
 
 	/** Initialize the camera extrinsics (rotation and translation vectors) for each image using linear algebra techniques */
-	static void InitCameraExtrinsics(
+	void InitCameraExtrinsics(
 		const TSubclassOf<ULensModel> LensModel,
 		const cv::Mat& ObjectPoints,
 		const cv::Mat& ImagePoints,
@@ -131,7 +150,7 @@ private:
 		const ECalibrationFlags SolverFlags);
 
 	/** Project the input object points to 2D using the input camera intrinsics, extrinsics, and distortion parameters */
-	static void ProjectPoints(
+	void ProjectPoints(
 		const TSubclassOf<ULensModel> LensModel,
 		const cv::Mat& ObjectPoints,
 		const cv::Mat& Rotation,
@@ -142,7 +161,7 @@ private:
 		cv::Mat& ProjectedPoints);
 
 	/** Project the input object points to 2D using the input camera intrinsics, extrinsics, and distortion parameters */
-	static void ProjectPoints(
+	void ProjectPoints(
 		const TSubclassOf<ULensModel> LensModel,
 		const cv::Mat& ObjectPoints,
 		const cv::Mat& Rotation,
@@ -155,7 +174,7 @@ private:
 		ECalibrationFlags SolverFlags);
 
 	/** Project the input object points to 2D using the input camera intrinsics, extrinsics, and anamorphic distortion parameters */
-	static void ProjectPointsAnamorphic(
+	void ProjectPointsAnamorphic(
 		const cv::Mat& ObjectPoints,
 		const cv::Mat& Rotation,
 		const cv::Mat& Translation,
@@ -167,7 +186,7 @@ private:
 		ECalibrationFlags SolverFlags);
 
 	/** Project the input object points to 2D using the input camera intrinsics, extrinsics, and spherical distortion parameters */
-	static void ProjectPointsSpherical(
+	void ProjectPointsSpherical(
 		const cv::Mat& ObjectPoints,
 		const cv::Mat& Rotation,
 		const cv::Mat& Translation,
@@ -178,14 +197,14 @@ private:
 		ECalibrationFlags SolverFlags);
 
 	/** Copy the input 3D and 2D points to OpenCV matrices for ease of use with the solver */
-	static void GatherPoints(
+	void GatherPoints(
 		const TArray<FObjectPoints>& InObjectPointsArray,
 		const TArray<FImagePoints>& InImagePointsArray,
 		cv::Mat& ObjectPointsMat,
 		cv::Mat& ImagePointsMat);
 
 	/** Divide large jacobian matrix into smaller views for each of the parameter groups the solver will solve */
-	static void SubdivideJacobian(
+	void SubdivideJacobian(
 		const cv::Mat& Jacobian,
 		cv::Mat& JacRotation,
 		cv::Mat& JacTranslation,
@@ -197,6 +216,20 @@ private:
 #endif	// WITH_OPENCV
 };
 
+class FCameraCalibrationSolver
+{
+public:
+
+	/** Optimize the input nodal offset transform by running a downhill solver that attempts to minimize the reprojection error of the input points */
+	static double OptimizeNodalOffset(
+		const TArray<TArray<FVector>>& InObjectPoints,
+		const TArray<TArray<FVector2f>>& InImagePoints,
+		const FVector2D& InFocalLength,
+		const FVector2D& InImageCenter,
+		const TArray<FTransform>& InCameraPoses,
+		FTransform& InOutNodalOffset);
+};
+
 #if WITH_OPENCV
 
 /** Levenberg–Marquardt solver for solving least-squares problems */
diff --git a/Engine/Plugins/VirtualProduction/CameraCalibration/Source/CameraCalibrationEditor/Private/Calibrators/CameraLensDistortionAlgoAruco.cpp b/Engine/Plugins/VirtualProduction/CameraCalibration/Source/CameraCalibrationEditor/Private/Calibrators/CameraLensDistortionAlgoAruco.cpp
--- a/Engine/Plugins/VirtualProduction/CameraCalibration/Source/CameraCalibrationEditor/Private/Calibrators/CameraLensDistortionAlgoAruco.cpp
+++ b/Engine/Plugins/VirtualProduction/CameraCalibration/Source/CameraCalibrationEditor/Private/Calibrators/CameraLensDistortionAlgoAruco.cpp
@@ -38,7 +38,6 @@
 static TAutoConsoleVariable<bool> CVarFixExtrinsicsAruco(TEXT("LensDistortionAruco.FixExtrinsics"), false, TEXT("If true, the solver will fix the camera extrinsics to the user-provided camera poses"));
 static TAutoConsoleVariable<bool> CVarFixZeroDistortionAruco(TEXT("LensDistortionAruco.FixZeroDistortion"), false, TEXT("If true, the solver will fix all distortion values to always be 0"));
 static TAutoConsoleVariable<bool> CVarUseExtrinsicsGuessAruco(TEXT("LensDistortionAruco.UseExtrinsicsGuess"), false, TEXT("If true, the actual calibrator and camera poses will be used when running the solver"));
-static TAutoConsoleVariable<bool> CVarUseNeuralNetSolverAruco(TEXT("LensDistortionAruco.UseNeuralNetSolver"), false, TEXT("If true, the neural network based solver will solver for distortion and camera intrinsics. Otherwise, the traditional OpenCV solver will be used."));
 #endif
 
 const int UCameraLensDistortionAlgoAruco::DATASET_VERSION = 1;
@@ -367,101 +366,68 @@ FDistortionCalibrationTask UCameraLensDistortionAlgoAruco::BeginCalibration()
 		CameraPoses.Add(Row->CameraPose);
 	}
 
-	const bool bUseNeuralNetSolver = CVarUseNeuralNetSolverAruco.GetValueOnAnyThread();
+	ECalibrationFlags SolverFlags = ECalibrationFlags::None;
 
-	if (!bUseNeuralNetSolver)
-	{
-		ECalibrationFlags SolverFlags = ECalibrationFlags::None;
-
-		// Aruco markers may be detected anywhere in the image, and there is no guarantee that they will be coplanar. 
-		// The solver's initialization for focal length assumes all points in an image are coplanar. 
-		// Therefore, we must provide an intrinsics guess to skip this initialization step in the solver. 
-		EnumAddFlags(SolverFlags, ECalibrationFlags::UseIntrinsicGuess);
-
-		if (CVarUseExtrinsicsGuessAruco.GetValueOnGameThread())
-		{
-			EnumAddFlags(SolverFlags, ECalibrationFlags::UseExtrinsicGuess);
-		}
+	// Aruco markers may be detected anywhere in the image, and there is no guarantee that they will be coplanar. 
+	// The solver's initialization for focal length assumes all points in an image are coplanar. 
+	// Therefore, we must provide an intrinsics guess to skip this initialization step in the solver. 
+	EnumAddFlags(SolverFlags, ECalibrationFlags::UseIntrinsicGuess);
 
-		if (CVarFixExtrinsicsAruco.GetValueOnAnyThread())
-		{
-			EnumAddFlags(SolverFlags, ECalibrationFlags::FixExtrinsics);
-		}
+	if (CVarUseExtrinsicsGuessAruco.GetValueOnGameThread())
+	{
+		EnumAddFlags(SolverFlags, ECalibrationFlags::UseExtrinsicGuess);
+	}
 
-		if (CVarFixZeroDistortionAruco.GetValueOnAnyThread())
-		{
-			EnumAddFlags(SolverFlags, ECalibrationFlags::FixZeroDistortion);
-		}
+	if (CVarFixExtrinsicsAruco.GetValueOnAnyThread())
+	{
+		EnumAddFlags(SolverFlags, ECalibrationFlags::FixExtrinsics);
+	}
 
-		if (bFixFocalLength)
-		{
-			EnumAddFlags(SolverFlags, ECalibrationFlags::FixFocalLength);
-		}
+	if (CVarFixZeroDistortionAruco.GetValueOnAnyThread())
+	{
+		EnumAddFlags(SolverFlags, ECalibrationFlags::FixZeroDistortion);
+	}
 
-		if (bFixImageCenter)
-		{
-			EnumAddFlags(SolverFlags, ECalibrationFlags::FixPrincipalPoint);
-		}
+	if (bFixFocalLength)
+	{
+		EnumAddFlags(SolverFlags, ECalibrationFlags::FixFocalLength);
+	}
 
-		const TSubclassOf<ULensModel> Model = LensFile->LensInfo.LensModel;
+	if (bFixImageCenter)
+	{
+		EnumAddFlags(SolverFlags, ECalibrationFlags::FixPrincipalPoint);
+	}
 
-		TArray<float> DistortionCoefficients;
+	const TSubclassOf<ULensModel> Model = LensFile->LensInfo.LensModel;
 
-		CalibrationTask = UE::Tasks::Launch(UE_SOURCE_LOCATION, [Model, Samples3d, Samples2d, ImageSize, FocalLength, ImageCenter, CameraPoses, PixelAspect, SolverFlags, Focus, Zoom]() mutable
-			{
-				FDistortionCalibrationResult Result;
-
-				Result.ReprojectionError = FCameraCalibrationSolver::CalibrateCamera(
-					Model,
-					Samples3d,
-					Samples2d,
-					ImageSize,
-					FocalLength,
-					ImageCenter,
-					Result.Parameters.Parameters,
-					CameraPoses,
-					PixelAspect,
-					SolverFlags
-				);
-
-				// CalibrateCamera() returns focal length and image center in pixels, but the result is expected to be normalized by the image size
-				Result.FocalLength.FxFy = FocalLength / ImageSize;
-				Result.ImageCenter.PrincipalPoint = ImageCenter / ImageSize;
-
-				// FZ inputs to LUT
-				Result.EvaluatedFocus = Focus;
-				Result.EvaluatedZoom = Zoom;
-
-				return Result;
-			});
-	}
-	else
-	{
-		// ULensDistortionSolver is the base class for the python implementation of the neural network solver.
-		// The intent is for there to be only one derived class, but to be sure, the name is checked to match
-		// the derived class that is expected to be found.
-		TArray<UClass*> DerivedSolverClasses;
-		GetDerivedClasses(ULensDistortionSolver::StaticClass(), DerivedSolverClasses);
+	UClass* SolverClass = LensDistortionTool->GetSolverClass();
 
-		if (DerivedSolverClasses.Num() == 0)
+	CalibrationTask = UE::Tasks::Launch(UE_SOURCE_LOCATION, [SolverClass, Model, Samples3d, Samples2d, ImageSize, FocalLength, ImageCenter, CameraPoses, PixelAspect, SolverFlags, Focus, Zoom]() mutable
 		{
-			UE_LOG(LogCameraCalibrationEditor, Error, TEXT("Could not initiate distortion calibration. No solver class was found."));
-			return CalibrationTask;
-		}
-
-		ULensDistortionSolver* TestSolver = Cast<ULensDistortionSolver>(DerivedSolverClasses.Last()->GetDefaultObject());
+			ULensDistortionSolver* Solver = NewObject<ULensDistortionSolver>(GetTransientPackage(), SolverClass);
+
+			FDistortionCalibrationResult Result = Solver->Solve(
+				Samples3d,
+				Samples2d,
+				ImageSize,
+				FocalLength,
+				ImageCenter,
+				CameraPoses,
+				Model,
+				PixelAspect,
+				SolverFlags
+			);
 
-		CalibrationTask = UE::Tasks::Launch(UE_SOURCE_LOCATION, [TestSolver, Samples3d, Samples2d, ImageSize, FocalLength, ImageCenter, Focus, Zoom]()
-			{
-				FDistortionCalibrationResult Result = TestSolver->Solve(Samples3d, Samples2d, ImageSize, FocalLength, ImageCenter);
+			// CalibrateCamera() returns focal length and image center in pixels, but the result is expected to be normalized by the image size
+			Result.FocalLength.FxFy = Result.FocalLength.FxFy / ImageSize;
+			Result.ImageCenter.PrincipalPoint = Result.ImageCenter.PrincipalPoint / ImageSize;
 
-				// FZ inputs to LUT
-				Result.EvaluatedFocus = Focus;
-				Result.EvaluatedFocus = Zoom;
+			// FZ inputs to LUT
+			Result.EvaluatedFocus = Focus;
+			Result.EvaluatedZoom = Zoom;
 
-				return Result;
-			});
-	}
+			return Result;
+		});
 
 	return CalibrationTask;
 }
diff --git a/Engine/Plugins/VirtualProduction/CameraCalibration/Source/CameraCalibrationEditor/Private/Calibrators/CameraLensDistortionAlgoCheckerboard.cpp b/Engine/Plugins/VirtualProduction/CameraCalibration/Source/CameraCalibrationEditor/Private/Calibrators/CameraLensDistortionAlgoCheckerboard.cpp
--- a/Engine/Plugins/VirtualProduction/CameraCalibration/Source/CameraCalibrationEditor/Private/Calibrators/CameraLensDistortionAlgoCheckerboard.cpp
+++ b/Engine/Plugins/VirtualProduction/CameraCalibration/Source/CameraCalibrationEditor/Private/Calibrators/CameraLensDistortionAlgoCheckerboard.cpp
@@ -41,7 +41,6 @@ static TAutoConsoleVariable<bool> CVarUseIntrinsicsGuess(TEXT("LensDistortionChe
 static TAutoConsoleVariable<bool> CVarFixExtrinsics(TEXT("LensDistortionCheckerboard.FixExtrinsics"), false, TEXT("If true, the solver will fix the camera extrinsics to the user-provided camera poses"));
 static TAutoConsoleVariable<bool> CVarFixZeroDistortion(TEXT("LensDistortionCheckerboard.FixZeroDistortion"), false, TEXT("If true, the solver will fix all distortion values to always be 0"));
 static TAutoConsoleVariable<bool> CVarUseExtrinsicsGuess(TEXT("LensDistortionCheckerboard.UseExtrinsicsGuess"), false, TEXT("If true, the actual checkerboard and camera poses will be used when running the solver"));
-static TAutoConsoleVariable<bool> CVarUseNeuralNetSolver(TEXT("LensDistortionCheckerboard.UseNeuralNetSolver"), false, TEXT("If true, the neural network based solver will solver for distortion and camera intrinsics. Otherwise, the traditional OpenCV solver will be used."));
 #endif
 
 const int UCameraLensDistortionAlgoCheckerboard::DATASET_VERSION = 1;
@@ -702,99 +701,68 @@ FDistortionCalibrationTask UCameraLensDistortionAlgoCheckerboard::BeginCalibrati
 		CameraPoses.Add(Row->CameraData.Pose);
 	}
 
-	const bool bUseNeuralNetSolver = CVarUseNeuralNetSolver.GetValueOnAnyThread();
+	ECalibrationFlags SolverFlags = ECalibrationFlags::None;
 
-	if (!bUseNeuralNetSolver)
+	if (CVarUseExtrinsicsGuess.GetValueOnGameThread())
 	{
-		ECalibrationFlags SolverFlags = ECalibrationFlags::None;
-
-		if (CVarUseExtrinsicsGuess.GetValueOnGameThread())
-		{
-			EnumAddFlags(SolverFlags, ECalibrationFlags::UseExtrinsicGuess);
-		}
-
-		if (CVarUseIntrinsicsGuess.GetValueOnGameThread())
-		{
-			EnumAddFlags(SolverFlags, ECalibrationFlags::UseIntrinsicGuess);
-		}
-
-		if (CVarFixExtrinsics.GetValueOnGameThread())
-		{
-			EnumAddFlags(SolverFlags, ECalibrationFlags::FixExtrinsics);
-		}
-
-		if (CVarFixZeroDistortion.GetValueOnGameThread())
-		{
-			EnumAddFlags(SolverFlags, ECalibrationFlags::FixZeroDistortion);
-		}
+		EnumAddFlags(SolverFlags, ECalibrationFlags::UseExtrinsicGuess);
+	}
 
-		if (bFixFocalLength)
-		{
-			EnumAddFlags(SolverFlags, ECalibrationFlags::FixFocalLength);
-		}
+	if (CVarUseIntrinsicsGuess.GetValueOnGameThread())
+	{
+		EnumAddFlags(SolverFlags, ECalibrationFlags::UseIntrinsicGuess);
+	}
 
-		if (bFixImageCenter)
-		{
-			EnumAddFlags(SolverFlags, ECalibrationFlags::FixPrincipalPoint);
-		}
+	if (CVarFixExtrinsics.GetValueOnGameThread())
+	{
+		EnumAddFlags(SolverFlags, ECalibrationFlags::FixExtrinsics);
+	}
 
-		const TSubclassOf<ULensModel> Model = LensFile->LensInfo.LensModel;
+	if (CVarFixZeroDistortion.GetValueOnGameThread())
+	{
+		EnumAddFlags(SolverFlags, ECalibrationFlags::FixZeroDistortion);
+	}
 
-		CalibrationTask = UE::Tasks::Launch(UE_SOURCE_LOCATION, [Model, Samples3d, Samples2d, ImageSize, FocalLength, ImageCenter, CameraPoses, PixelAspect, SolverFlags, Focus, Zoom]() mutable
-			{
-				FDistortionCalibrationResult Result;
-
-				Result.ReprojectionError = FCameraCalibrationSolver::CalibrateCamera(
-					Model,
-					Samples3d,
-					Samples2d,
-					ImageSize,
-					FocalLength,
-					ImageCenter,
-					Result.Parameters.Parameters,
-					CameraPoses,
-					PixelAspect,
-					SolverFlags
-				);
-
-				// CalibrateCamera() returns focal length and image center in pixels, but the result is expected to be normalized by the image size
-				Result.FocalLength.FxFy = FocalLength / ImageSize;
-				Result.ImageCenter.PrincipalPoint = ImageCenter / ImageSize;
-
-				// FZ inputs to LUT
-				Result.EvaluatedFocus = Focus;
-				Result.EvaluatedZoom = Zoom;
-
-				return Result;
-			});
+	if (bFixFocalLength)
+	{
+		EnumAddFlags(SolverFlags, ECalibrationFlags::FixFocalLength);
 	}
-	else
+
+	if (bFixImageCenter)
 	{
-		// ULensDistortionSolver is the base class for the python implementation of the neural network solver.
-		// The intent is for there to be only one derived class, but to be sure, the name is checked to match
-		// the derived class that is expected to be found.
-		TArray<UClass*> DerivedSolverClasses;
-		GetDerivedClasses(ULensDistortionSolver::StaticClass(), DerivedSolverClasses);
+		EnumAddFlags(SolverFlags, ECalibrationFlags::FixPrincipalPoint);
+	}
 
-		if (DerivedSolverClasses.Num() == 0)
-		{
-			UE_LOG(LogCameraCalibrationEditor, Error, TEXT("Could not initiate distortion calibration. No solver class was found."));
-			return CalibrationTask;
-		}
+	const TSubclassOf<ULensModel> Model = LensFile->LensInfo.LensModel;
+	
+	UClass* SolverClass = Tool->GetSolverClass();
 
-		ULensDistortionSolver* TestSolver = Cast<ULensDistortionSolver>(DerivedSolverClasses.Last()->GetDefaultObject());
+	CalibrationTask = UE::Tasks::Launch(UE_SOURCE_LOCATION, [SolverClass, Model, Samples3d, Samples2d, ImageSize, FocalLength, ImageCenter, CameraPoses, PixelAspect, SolverFlags, Focus, Zoom]() mutable
+		{
+			ULensDistortionSolver* Solver = NewObject<ULensDistortionSolver>(GetTransientPackage(), SolverClass);
+
+			FDistortionCalibrationResult Result = Solver->Solve(
+				Samples3d,
+				Samples2d,
+				ImageSize,
+				FocalLength,
+				ImageCenter,
+				CameraPoses,
+				Model,
+				PixelAspect,
+				SolverFlags
+			);
 
-		CalibrationTask = UE::Tasks::Launch(UE_SOURCE_LOCATION, [TestSolver, Samples3d, Samples2d, ImageSize, FocalLength, ImageCenter, Focus, Zoom]()
-			{ 
-				FDistortionCalibrationResult Result = TestSolver->Solve(Samples3d, Samples2d, ImageSize, FocalLength, ImageCenter);
+			// CalibrateCamera() returns focal length and image center in pixels, but the result is expected to be normalized by the image size
+			Result.FocalLength.FxFy = Result.FocalLength.FxFy / ImageSize;
+			Result.ImageCenter.PrincipalPoint = Result.ImageCenter.PrincipalPoint / ImageSize;
 
-				// FZ inputs to LUT
-				Result.EvaluatedFocus = Focus;
-				Result.EvaluatedFocus = Zoom;
+			// FZ inputs to LUT
+			Result.EvaluatedFocus = Focus;
+			Result.EvaluatedZoom = Zoom;
 
-				return Result;
-			});
-	}
+			return Result;
+		});
 
 	return CalibrationTask;
 }
diff --git a/Engine/Plugins/VirtualProduction/CameraCalibration/Source/CameraCalibrationEditor/Private/Calibrators/CameraLensDistortionAlgoPoints.cpp b/Engine/Plugins/VirtualProduction/CameraCalibration/Source/CameraCalibrationEditor/Private/Calibrators/CameraLensDistortionAlgoPoints.cpp
--- a/Engine/Plugins/VirtualProduction/CameraCalibration/Source/CameraCalibrationEditor/Private/Calibrators/CameraLensDistortionAlgoPoints.cpp
+++ b/Engine/Plugins/VirtualProduction/CameraCalibration/Source/CameraCalibrationEditor/Private/Calibrators/CameraLensDistortionAlgoPoints.cpp
@@ -472,29 +472,35 @@ bool UCameraLensDistortionAlgoPoints::GetLensDistortion(
 	ECalibrationFlags SolverFlags = ECalibrationFlags::None;
 	EnumAddFlags(SolverFlags, ECalibrationFlags::UseIntrinsicGuess);
 
-	TArray<float> DistortionCoefficients;
+	UClass* SolverClass = LensDistortionTool->GetSolverClass();
+	ULensDistortionSolver* Solver = NewObject<ULensDistortionSolver>(this, SolverClass);
 
-	OutError = FCameraCalibrationSolver::CalibrateCamera(
-		LensFile->LensInfo.LensModel,
+	FDistortionCalibrationResult Result = Solver->Solve(
 		Samples3d,
 		Samples2d,
 		ImageSize,
 		FocalLength,
 		ImageCenter,
-		DistortionCoefficients,
 		CameraPoses,
+		LensFile->LensInfo.LensModel,
 		PixelAspect,
 		SolverFlags
 	);
 
+	if (!Result.ErrorMessage.IsEmpty())
+	{
+		OutErrorMessage = Result.ErrorMessage;
+		return false;
+	}
+
 	const TSharedPtr<FLensDistortionPointsRowData>& FirstRow = CalibrationRows[0];
 
 	OutLensModel = LensFile->LensInfo.LensModel;
-	OutDistortionInfo.Parameters = DistortionCoefficients;
+	OutDistortionInfo.Parameters = Result.Parameters.Parameters;
 
-	OutFocalLengthInfo.FxFy = FVector2D(FocalLength / ImageSize);
+	OutFocalLengthInfo.FxFy = FVector2D(Result.FocalLength.FxFy / ImageSize);
 
-	OutImageCenterInfo.PrincipalPoint = FVector2D(ImageCenter / ImageSize);
+	OutImageCenterInfo.PrincipalPoint = FVector2D(Result.ImageCenter.PrincipalPoint / ImageSize);
 
 	// FZ inputs to LUT
 	OutFocus = FirstRow->CameraData.InputFocus;
diff --git a/Engine/Plugins/VirtualProduction/CameraCalibration/Source/CameraCalibrationEditor/Private/TestCameraCalibration.cpp b/Engine/Plugins/VirtualProduction/CameraCalibration/Source/CameraCalibrationEditor/Private/TestCameraCalibration.cpp
--- a/Engine/Plugins/VirtualProduction/CameraCalibration/Source/CameraCalibrationEditor/Private/TestCameraCalibration.cpp
+++ b/Engine/Plugins/VirtualProduction/CameraCalibration/Source/CameraCalibrationEditor/Private/TestCameraCalibration.cpp
@@ -6,6 +6,7 @@
 #include "CoreMinimal.h"
 
 #include "Calibrators/CameraCalibrationSolver.h"
+#include "CameraCalibrationTypes.h"
 #include "CameraCalibrationUtilsPrivate.h"
 #include "Dialog/SCustomDialog.h"
 #include "Engine/Engine.h"
@@ -255,15 +256,16 @@ namespace UE::Private::CameraCalibration::AutomatedTests
 		FVector2D CalibratedImageCenter = ImageCenter;
 		TArray<float> CalibratedDistortionParameters;
 
-		const double Error = FCameraCalibrationSolver::CalibrateCamera(
-			USphericalLensModel::StaticClass(),
+		ULensDistortionSolverOpenCV* TestSolver = NewObject<ULensDistortionSolverOpenCV>();
+
+		FDistortionCalibrationResult Result = TestSolver->Solve(
 			NoisyObjectPoints,
 			NoisyImagePoints,
 			ImageSize,
 			CalibratedFxFy,
 			CalibratedImageCenter,
-			CalibratedDistortionParameters,
 			EstimatedCameraPoses,
+			USphericalLensModel::StaticClass(),
 			PixelAspect,
 			SolverFlags
 		);
@@ -287,48 +289,20 @@ namespace UE::Private::CameraCalibration::AutomatedTests
 		Test.AddInfo(FString::Printf(TEXT("\t\tP1: %f"), SphericalParams.P1));
 		Test.AddInfo(FString::Printf(TEXT("\t\tP2: %f"), SphericalParams.P2));
 
-		Test.AddInfo(TEXT("Ground-Truth Camera Poses:"));
-		for (int32 PoseIndex = 0; PoseIndex < TrueCameraPoses.Num(); ++PoseIndex)
-		{
-			const FTransform& Pose = TrueCameraPoses[PoseIndex];
-
-			Test.AddInfo(FString::Printf(TEXT("\t\tPose %d"), PoseIndex));
-
-			const FVector Translation = Pose.GetTranslation();
-			const FRotator Rotator = Pose.GetRotation().Rotator();
-
-			Test.AddInfo(FString::Printf(TEXT("\t\t\t\tTranslation: (%lf, %lf, %lf)"), Translation.X, Translation.Y, Translation.Z));
-			Test.AddInfo(FString::Printf(TEXT("\t\t\t\tRotation:    (%lf, %lf, %lf)"), Rotator.Roll, Rotator.Pitch, Rotator.Yaw));
-		}
-
 		Test.AddInfo(TEXT("\n"));
- 		Test.AddInfo(FString::Printf(TEXT("Result RMS Error: %lf"), Error));
+ 		Test.AddInfo(FString::Printf(TEXT("Result RMS Error: %lf"), Result.ReprojectionError));
 		Test.AddInfo(TEXT("\n"));
 
 		Test.AddInfo(TEXT("Calibrated Camera Intrinsics:"));
- 		Test.AddInfo(FString::Printf(TEXT("\t\tFocal Length: %lf mm (%lf pixels)"), (CalibratedFxFy.X / ImageSize.X)* SensorWidth, CalibratedFxFy.X));
-		Test.AddInfo(FString::Printf(TEXT("\t\tImage Center: (%lf, %lf)"), CalibratedImageCenter.X, CalibratedImageCenter.Y));
+ 		Test.AddInfo(FString::Printf(TEXT("\t\tFocal Length: %lf mm (%lf pixels)"), (Result.FocalLength.FxFy.X / ImageSize.X)* SensorWidth, Result.FocalLength.FxFy.X));
+		Test.AddInfo(FString::Printf(TEXT("\t\tImage Center: (%lf, %lf)"), Result.ImageCenter.PrincipalPoint.X, Result.ImageCenter.PrincipalPoint.Y));
 
  		Test.AddInfo(TEXT("Calibrated Distortion Coefficients:"));
-		Test.AddInfo(FString::Printf(TEXT("\t\tK1: %f"), CalibratedDistortionParameters[0]));
-		Test.AddInfo(FString::Printf(TEXT("\t\tK2: %f"), CalibratedDistortionParameters[1]));
-		Test.AddInfo(FString::Printf(TEXT("\t\tK3: %f"), CalibratedDistortionParameters[2]));
-		Test.AddInfo(FString::Printf(TEXT("\t\tP1: %f"), CalibratedDistortionParameters[3]));
-		Test.AddInfo(FString::Printf(TEXT("\t\tP2: %f"), CalibratedDistortionParameters[4]));
-
-		Test.AddInfo(TEXT("Calibrated Camera Poses:"));
-		for (int32 PoseIndex = 0; PoseIndex < EstimatedCameraPoses.Num(); ++PoseIndex)
-		{
-			const FTransform& Pose = EstimatedCameraPoses[PoseIndex];
-
-			Test.AddInfo(FString::Printf(TEXT("\t\tPose %d"), PoseIndex));
-
-			const FVector Translation = Pose.GetTranslation();
-			const FRotator Rotator = Pose.GetRotation().Rotator();
-
-			Test.AddInfo(FString::Printf(TEXT("\t\t\t\tTranslation: (%lf, %lf, %lf)"), Translation.X, Translation.Y, Translation.Z));
-			Test.AddInfo(FString::Printf(TEXT("\t\t\t\tRotation:    (%lf, %lf, %lf)"), Rotator.Roll, Rotator.Pitch, Rotator.Yaw));
-		}
+		Test.AddInfo(FString::Printf(TEXT("\t\tK1: %f"), Result.Parameters.Parameters[0]));
+		Test.AddInfo(FString::Printf(TEXT("\t\tK2: %f"), Result.Parameters.Parameters[1]));
+		Test.AddInfo(FString::Printf(TEXT("\t\tK3: %f"), Result.Parameters.Parameters[2]));
+		Test.AddInfo(FString::Printf(TEXT("\t\tP1: %f"), Result.Parameters.Parameters[3]));
+		Test.AddInfo(FString::Printf(TEXT("\t\tP2: %f"), Result.Parameters.Parameters[4]));
 	}
 
 	void TestNodalOffsetCalibration(FAutomationTestBase& Test)
