CameraCalibration: Add ability to cancel an async calibration task. Add ability to query an in-progress calibration task for its current status.

diff --git a/Engine/Plugins/VirtualProduction/CameraCalibration/Source/CameraCalibrationEditor/Private/AssetEditor/LensDistortionTool.cpp b/Engine/Plugins/VirtualProduction/CameraCalibration/Source/CameraCalibrationEditor/Private/AssetEditor/LensDistortionTool.cpp
--- a/Engine/Plugins/VirtualProduction/CameraCalibration/Source/CameraCalibrationEditor/Private/AssetEditor/LensDistortionTool.cpp
+++ b/Engine/Plugins/VirtualProduction/CameraCalibration/Source/CameraCalibrationEditor/Private/AssetEditor/LensDistortionTool.cpp
@@ -107,6 +107,13 @@ void ULensDistortionTool::Shutdown()
 {
 	if (CurrentAlgo)
 	{
+		if (CalibrationTask.IsValid())
+		{
+			CurrentAlgo->CancelCalibration();
+			CalibrationTask = {};
+			ProgressWindow->HideWindow();
+		}
+
 		CurrentAlgo->Shutdown();
 		CurrentAlgo = nullptr;
 
@@ -156,6 +163,16 @@ void ULensDistortionTool::Tick(float DeltaTime)
 
 			OkayButton->SetEnabled(true);
 		}
+		else
+		{
+			FText StatusText = FText::GetEmpty();
+			const bool bIsStatusNew = CurrentAlgo->GetCalibrationStatus(StatusText);
+
+			if (bIsStatusNew)
+			{
+				ProgressTextWidget->SetText(StatusText);
+			}
+		}
 	}
 }
 
@@ -171,6 +188,12 @@ bool ULensDistortionTool::OnViewportClicked(const FGeometry& MyGeometry, const F
 		return false;
 	}
 
+	// Block user interaction with the simulcam viewport while an async calibration task is executing
+	if (CalibrationTask.IsValid())
+	{
+		return false;
+	}
+
 	return CurrentAlgo->OnViewportClicked(MyGeometry, MouseEvent);
 }
 
@@ -185,6 +208,7 @@ void ULensDistortionTool::BuildProgressWindowWidgets()
 	ProgressWindow = SNew(SWindow)
 		.Title(LOCTEXT("ProgressWindowTitle", "Distortion Calibration Progress"))
 		.SizingRule(ESizingRule::Autosized)
+		.IsTopmostWindow(true)
 		.HasCloseButton(false)
 		.SupportsMaximize(false)
 		.SupportsMinimize(true);
@@ -378,6 +402,8 @@ void ULensDistortionTool::OnSaveCurrentCalibrationData()
 		// Ensure that the Ok button is disabled and show the progress window
 		OkayButton->SetEnabled(false);
 		ProgressWindow->ShowWindow();
+
+		DistortionWidget->SetEnabled(false);
 	}
 	else
 	{
@@ -418,7 +444,11 @@ void ULensDistortionTool::OnSaveCurrentCalibrationData()
 
 FReply ULensDistortionTool::OnCancelPressed()
 {
-	// TODO: Implement correct cancel handling. Send a signal to the algo to interrupt the currently running calibration task 
+	CurrentAlgo->CancelCalibration();
+
+	CalibrationTask = {};
+	ProgressWindow->HideWindow();
+	DistortionWidget->SetEnabled(true);
 
 	return FReply::Handled();
 }
@@ -428,6 +458,7 @@ FReply ULensDistortionTool::OnOkPressed()
 	SaveCalibrationResult();
 
 	ProgressWindow->HideWindow();
+	DistortionWidget->SetEnabled(true);
 
 	return FReply::Handled();
 }
diff --git a/Engine/Plugins/VirtualProduction/CameraCalibration/Source/CameraCalibrationEditor/Private/Calibrators/CameraCalibrationSolver.cpp b/Engine/Plugins/VirtualProduction/CameraCalibration/Source/CameraCalibrationEditor/Private/Calibrators/CameraCalibrationSolver.cpp
--- a/Engine/Plugins/VirtualProduction/CameraCalibration/Source/CameraCalibrationEditor/Private/Calibrators/CameraCalibrationSolver.cpp
+++ b/Engine/Plugins/VirtualProduction/CameraCalibration/Source/CameraCalibrationEditor/Private/Calibrators/CameraCalibrationSolver.cpp
@@ -23,6 +23,23 @@ static TAutoConsoleVariable<float> CVarLocationStepValue(TEXT("CameraCalibration
 
 DEFINE_LOG_CATEGORY_STATIC(LogCameraCalibrationSolver, Log, All);
 
+bool ULensDistortionSolver::GetStatusText(FText& OutStatusText) 
+{ 
+	OutStatusText = StatusText;
+	if (bHasStatusChanged)
+	{
+		bHasStatusChanged = false;
+		return true;
+	}
+	return false;
+}
+
+void ULensDistortionSolver::SetStatusText(FText InStatusText) 
+{ 
+	StatusText = MoveTemp(InStatusText);
+	bHasStatusChanged = true;
+}
+
 FText ULensDistortionSolverOpenCV::GetDisplayName_Implementation() const
 {
 	return LOCTEXT("OpenCVSolverDisplayName", "OpenCV Solver");
@@ -300,10 +317,11 @@ FDistortionCalibrationResult ULensDistortionSolverOpenCV::Solve_Implementation(
 	cv::Mat Jacobian(MaxPoints * 2, NumExtrinsics + NumIntrinsics, CV_64FC1, cv::Scalar(0));
 	cv::Mat Diffs(MaxPoints * 2, 1, CV_64FC1);
 
+	int32 LoopCounter = 0;
 	while (true)
 	{
 		// Update the solver
-		bool bShouldProceed = Solver.UpdateAlt();
+		bool bShouldProceed = Solver.UpdateAlt() && IsRunning();
 		bool bComputeJacobian = Solver.State == FLevMarqSolver::ESolverState::ComputeJacobian;
 
 		// Update the camera matrix and distortion parameters with the latest values of the parameters from the solver
@@ -407,6 +425,11 @@ FDistortionCalibrationResult ULensDistortionSolverOpenCV::Solve_Implementation(
 
 		// Update the solver's error with the latest reprojection error
 		Solver.ErrorNorm = ReprojectionError;
+
+		const double CurrentRMSE = FMath::Sqrt(ReprojectionError / NumTotalPoints);
+		SetStatusText(FText::Format(LOCTEXT("ReprojectionError", "Reprojection Error: {0} pixels for loop {1}"), CurrentRMSE, LoopCounter));
+
+		++LoopCounter;
 	}
 
 	RMSE = FMath::Sqrt(ReprojectionError / NumTotalPoints);
diff --git a/Engine/Plugins/VirtualProduction/CameraCalibration/Source/CameraCalibrationEditor/Private/Calibrators/CameraCalibrationSolver.h b/Engine/Plugins/VirtualProduction/CameraCalibration/Source/CameraCalibrationEditor/Private/Calibrators/CameraCalibrationSolver.h
--- a/Engine/Plugins/VirtualProduction/CameraCalibration/Source/CameraCalibrationEditor/Private/Calibrators/CameraCalibrationSolver.h
+++ b/Engine/Plugins/VirtualProduction/CameraCalibration/Source/CameraCalibrationEditor/Private/Calibrators/CameraCalibrationSolver.h
@@ -85,6 +85,13 @@ public:
 	UFUNCTION(BlueprintNativeEvent, Category = "Calibration")
 	FText GetDisplayName() const;
 
+public:
+	/** Get the latest status of the solve and marks the status as old. Returns true if this status was new, or false if this status was old. */
+	bool GetStatusText(FText& OutStatusText);
+
+	/** Cancels the solve by setting bIsRunning to false, inidicating that the solver should stop executing the next time it checks IsRunning() */
+	void Cancel() { bIsRunning = false; }
+
 	virtual FDistortionCalibrationResult Solve_Implementation(
 		const TArray<FObjectPoints>& ObjectPointArray,
 		const TArray<FImagePoints>& ImagePointArray,
@@ -97,6 +104,23 @@ public:
 		ECalibrationFlags SolverFlags) PURE_VIRTUAL(ULensDistortionSolver::Solve_Implementation, return FDistortionCalibrationResult(););
 
 	virtual FText GetDisplayName_Implementation() const PURE_VIRTUAL(ULensDistortionSolver::GetDisplayName_Implementation, return FText::GetEmpty(););
+
+protected:
+	/** Returns true if the solver is currently running, and false if it has been cancelled. The solver should call this in critical loops in order to respond to cancellation requests. */
+	bool IsRunning() const { return bIsRunning; }
+
+	/** Sets the latest status text and marks bIsStatusNew to true */
+	void SetStatusText(FText InStatusText);
+
+protected:
+	/** True is the solver should continue executing. Set to false by Cancel(), indicating that the solver should early-out the next time it checks IsRunning() */
+	std::atomic<bool> bIsRunning = true;
+
+	/** Status text describing the current state of the solve */
+	FText StatusText;
+
+	/** Set to true by SetStatusText() when a new status is available. Set to false by GetStatusText(), indicating that the latest status has already been queried. */
+	bool bHasStatusChanged = false;
 };
 
 /** 
diff --git a/Engine/Plugins/VirtualProduction/CameraCalibration/Source/CameraCalibrationEditor/Private/Calibrators/CameraLensDistortionAlgoCheckerboard.cpp b/Engine/Plugins/VirtualProduction/CameraCalibration/Source/CameraCalibrationEditor/Private/Calibrators/CameraLensDistortionAlgoCheckerboard.cpp
--- a/Engine/Plugins/VirtualProduction/CameraCalibration/Source/CameraCalibrationEditor/Private/Calibrators/CameraLensDistortionAlgoCheckerboard.cpp
+++ b/Engine/Plugins/VirtualProduction/CameraCalibration/Source/CameraCalibrationEditor/Private/Calibrators/CameraLensDistortionAlgoCheckerboard.cpp
@@ -735,12 +735,10 @@ FDistortionCalibrationTask UCameraLensDistortionAlgoCheckerboard::BeginCalibrati
 
 	const TSubclassOf<ULensModel> Model = LensFile->LensInfo.LensModel;
 	
-	UClass* SolverClass = Tool->GetSolverClass();
+	Solver = NewObject<ULensDistortionSolver>(this, Tool->GetSolverClass());
 
-	CalibrationTask = UE::Tasks::Launch(UE_SOURCE_LOCATION, [SolverClass, Model, Samples3d, Samples2d, ImageSize, FocalLength, ImageCenter, CameraPoses, PixelAspect, SolverFlags, Focus, Zoom]() mutable
+	CalibrationTask = UE::Tasks::Launch(UE_SOURCE_LOCATION, [Solver=Solver, Model, Samples3d, Samples2d, ImageSize, FocalLength, ImageCenter, CameraPoses, PixelAspect, SolverFlags, Focus, Zoom]() mutable
 		{
-			ULensDistortionSolver* Solver = NewObject<ULensDistortionSolver>(GetTransientPackage(), SolverClass);
-
 			FDistortionCalibrationResult Result = Solver->Solve(
 				Samples3d,
 				Samples2d,
@@ -767,6 +765,23 @@ FDistortionCalibrationTask UCameraLensDistortionAlgoCheckerboard::BeginCalibrati
 	return CalibrationTask;
 }
 
+void UCameraLensDistortionAlgoCheckerboard::CancelCalibration()
+{
+	if (Solver)
+	{
+		Solver->Cancel();
+	}
+}
+
+bool UCameraLensDistortionAlgoCheckerboard::GetCalibrationStatus(FText& StatusText) const
+{
+	if (Solver)
+	{
+		return Solver->GetStatusText(StatusText);
+	}
+	return false;
+}
+
 TSharedRef<SWidget> UCameraLensDistortionAlgoCheckerboard::BuildCalibrationDevicePickerWidget()
 {
 	return SNew(SHorizontalBox)
diff --git a/Engine/Plugins/VirtualProduction/CameraCalibration/Source/CameraCalibrationEditor/Private/Calibrators/CameraLensDistortionAlgoCheckerboard.h b/Engine/Plugins/VirtualProduction/CameraCalibration/Source/CameraCalibrationEditor/Private/Calibrators/CameraLensDistortionAlgoCheckerboard.h
--- a/Engine/Plugins/VirtualProduction/CameraCalibration/Source/CameraCalibrationEditor/Private/Calibrators/CameraLensDistortionAlgoCheckerboard.h
+++ b/Engine/Plugins/VirtualProduction/CameraCalibration/Source/CameraCalibrationEditor/Private/Calibrators/CameraLensDistortionAlgoCheckerboard.h
@@ -3,6 +3,8 @@
 #pragma once
 
 #include "CameraLensDistortionAlgo.h"
+
+#include "CameraCalibrationSolver.h"
 #include "ImageCore.h"
 
 #include "CameraLensDistortionAlgoCheckerboard.generated.h"
@@ -134,6 +136,8 @@ public:
 	virtual bool IsOverlayEnabled() const override { return bShouldShowOverlay; };
 	virtual void OnDistortionSavedToLens() override;
 	virtual FDistortionCalibrationTask BeginCalibration() override;
+	virtual void CancelCalibration() override;
+	virtual bool GetCalibrationStatus(FText& StatusText) const override;
 	virtual bool SupportsAsyncCalibration() override { return true; };
 	virtual bool HasCalibrationData() const override;
 	virtual void PreImportCalibrationData() override;
@@ -190,6 +194,9 @@ private:
 	/** Texture into which detected chessboard corners from each calibration row are drawn */
 	TObjectPtr<UTexture2D> CoverageTexture;
 
+	/** Solver instance that will run the distortion calibration */
+	TObjectPtr<ULensDistortionSolver> Solver;
+
 private:
 
 	/** Builds the UI of the calibration device picker */
diff --git a/Engine/Plugins/VirtualProduction/CameraCalibrationCore/Source/CameraCalibrationCore/Public/CameraLensDistortionAlgo.h b/Engine/Plugins/VirtualProduction/CameraCalibrationCore/Source/CameraCalibrationCore/Public/CameraLensDistortionAlgo.h
--- a/Engine/Plugins/VirtualProduction/CameraCalibrationCore/Source/CameraCalibrationCore/Public/CameraLensDistortionAlgo.h
+++ b/Engine/Plugins/VirtualProduction/CameraCalibrationCore/Source/CameraCalibrationCore/Public/CameraLensDistortionAlgo.h
@@ -88,6 +88,12 @@ public:
 	/** Launches an asynchronous task to perform the distortion calibration. The task will contain the calibration result when it finishes. */
 	virtual FDistortionCalibrationTask BeginCalibration() { return {}; };
 
+	/** Cancel an in-progress async calibration task */
+	virtual void CancelCalibration() { };
+
+	/** Get the latest status from the lens distortion calibration. Returns true if this status is "new", or false if this status was previously reported. */
+	virtual bool GetCalibrationStatus(FText& StatusText) const { return false; };
+
 	/** 
 	 * If true, the algo must override BeginCalibration() and use it to launch the calibration as an asynchronous task.
 	 * If false, the algo must override GetLensDistortion() and run the calibration on the Game Thread.
