CameraCalibration: Check-in prototype/sample pytorch solver for lens distortion and camera intrinsics

diff --git a/Engine/Plugins/VirtualProduction/CameraCalibration/CameraCalibration.uplugin b/Engine/Plugins/VirtualProduction/CameraCalibration/CameraCalibration.uplugin
--- a/Engine/Plugins/VirtualProduction/CameraCalibration/CameraCalibration.uplugin
+++ b/Engine/Plugins/VirtualProduction/CameraCalibration/CameraCalibration.uplugin
@@ -40,6 +40,18 @@
 		{
 			"Name": "OpenCVLensDistortion",
 			"Enabled": true
+		},
+		{
+			"Name": "PythonFoundationPackages",
+			"Enabled": true
+		}
+	],
+	"PythonRequirements": [
+		{
+			"Platform": "All",
+			"Requirements": [
+				"nerfacc"
+			]
 		}
 	]
 }
\ No newline at end of file
diff --git a/Engine/Plugins/VirtualProduction/CameraCalibration/Source/CameraCalibrationEditor/Private/AssetEditor/SLensDistortionToolPanel.cpp b/Engine/Plugins/VirtualProduction/CameraCalibration/Source/CameraCalibrationEditor/Private/AssetEditor/SLensDistortionToolPanel.cpp
--- a/Engine/Plugins/VirtualProduction/CameraCalibration/Source/CameraCalibrationEditor/Private/AssetEditor/SLensDistortionToolPanel.cpp
+++ b/Engine/Plugins/VirtualProduction/CameraCalibration/Source/CameraCalibrationEditor/Private/AssetEditor/SLensDistortionToolPanel.cpp
@@ -236,18 +236,21 @@ TSharedRef<SWidget> SLensDistortionToolPanel::BuildSolverList()
 		{
 			ULensDistortionSolver* SolverCDO = Cast<ULensDistortionSolver>(SolverClass->GetDefaultObject());
 
-			MenuBuilder.AddMenuEntry(
-				SolverCDO->GetDisplayName(),
-				SolverCDO->GetDisplayName(), // TODO: If this remains a user-facing option, then it may be better to have the solver implement a class description text field to use as the tooltip instead
-				FSlateIcon(),
-				FUIAction(
-					FExecuteAction::CreateSP(this, &SLensDistortionToolPanel::OnSolverClassSelected, SolverClass),
-					FCanExecuteAction(),
-					FIsActionChecked::CreateSP(this, &SLensDistortionToolPanel::IsSolverClassSelected, SolverClass)
-				),
-				NAME_None,
-				EUserInterfaceActionType::RadioButton
-			);
+			if (SolverCDO->IsEnabled())
+			{
+				MenuBuilder.AddMenuEntry(
+					SolverCDO->GetDisplayName(),
+					SolverCDO->GetDisplayName(), // TODO: If this remains a user-facing option, then it may be better to have the solver implement a class description text field to use as the tooltip instead
+					FSlateIcon(),
+					FUIAction(
+						FExecuteAction::CreateSP(this, &SLensDistortionToolPanel::OnSolverClassSelected, SolverClass),
+						FCanExecuteAction(),
+						FIsActionChecked::CreateSP(this, &SLensDistortionToolPanel::IsSolverClassSelected, SolverClass)
+					),
+					NAME_None,
+					EUserInterfaceActionType::RadioButton
+				);
+			}
 		}
 	}
 
diff --git a/Engine/Plugins/VirtualProduction/CameraCalibration/Source/CameraCalibrationEditor/Private/Calibrators/CameraCalibrationSolver.h b/Engine/Plugins/VirtualProduction/CameraCalibration/Source/CameraCalibrationEditor/Private/Calibrators/CameraCalibrationSolver.h
--- a/Engine/Plugins/VirtualProduction/CameraCalibration/Source/CameraCalibrationEditor/Private/Calibrators/CameraCalibrationSolver.h
+++ b/Engine/Plugins/VirtualProduction/CameraCalibration/Source/CameraCalibrationEditor/Private/Calibrators/CameraCalibrationSolver.h
@@ -85,6 +85,10 @@ public:
 	UFUNCTION(BlueprintNativeEvent, Category = "Calibration")
 	FText GetDisplayName() const;
 
+	/** Get the name of this solver class for displaying in the editor UI */
+	UFUNCTION(BlueprintNativeEvent, Category = "Calibration")
+	bool IsEnabled() const;
+
 public:
 	/** Get the latest status of the solve and marks the status as old. Returns true if this status was new, or false if this status was old. */
 	bool GetStatusText(FText& OutStatusText);
@@ -105,8 +109,11 @@ public:
 
 	virtual FText GetDisplayName_Implementation() const PURE_VIRTUAL(ULensDistortionSolver::GetDisplayName_Implementation, return FText::GetEmpty(););
 
+	virtual bool IsEnabled_Implementation() const { return true; }
+
 protected:
 	/** Returns true if the solver is currently running, and false if it has been cancelled. The solver should call this in critical loops in order to respond to cancellation requests. */
+	UFUNCTION(BlueprintCallable, Category = "Calibration")
 	bool IsRunning() const { return bIsRunning; }
 
 	/** Sets the latest status text and marks bIsStatusNew to true */
