CameraCalibration: Fix error messaging in distortion tool
#rnx

diff --git a/Engine/Plugins/VirtualProduction/CameraCalibration/Source/CameraCalibrationEditor/Private/AssetEditor/LensDistortionTool.cpp b/Engine/Plugins/VirtualProduction/CameraCalibration/Source/CameraCalibrationEditor/Private/AssetEditor/LensDistortionTool.cpp
--- a/Engine/Plugins/VirtualProduction/CameraCalibration/Source/CameraCalibrationEditor/Private/AssetEditor/LensDistortionTool.cpp
+++ b/Engine/Plugins/VirtualProduction/CameraCalibration/Source/CameraCalibrationEditor/Private/AssetEditor/LensDistortionTool.cpp
@@ -384,15 +384,23 @@ void ULensDistortionTool::OnSaveCurrentCalibrationData()
 	}
 
 	const FText TitleError = LOCTEXT("LensCalibrationError", "Lens Calibration Error");
-	const FText UnknownError = LOCTEXT("UnknownError", "An error occurred initiating the distortion calibration. Check the output log for details.");
+	const FText UnknownError = LOCTEXT("UnknownError", "An unknown error occurred initiating the distortion calibration. Check the output log for details.");
 
 	if (CurrentAlgo->SupportsAsyncCalibration())
 	{
-		CalibrationTask = CurrentAlgo->BeginCalibration();
+		FText ErrorMessage;
+		CalibrationTask = CurrentAlgo->BeginCalibration(ErrorMessage);
 
 		if (!CalibrationTask.IsValid())
 		{
-			FMessageDialog::Open(EAppMsgType::Ok, UnknownError, TitleError);
+			if (!ErrorMessage.IsEmpty())
+			{
+				FMessageDialog::Open(EAppMsgType::Ok, ErrorMessage, TitleError);
+			}
+			else
+			{
+				FMessageDialog::Open(EAppMsgType::Ok, UnknownError, TitleError);
+			}
 			return;
 		}
 
diff --git a/Engine/Plugins/VirtualProduction/CameraCalibration/Source/CameraCalibrationEditor/Private/Calibrators/CameraLensDistortionAlgoAruco.cpp b/Engine/Plugins/VirtualProduction/CameraCalibration/Source/CameraCalibrationEditor/Private/Calibrators/CameraLensDistortionAlgoAruco.cpp
--- a/Engine/Plugins/VirtualProduction/CameraCalibration/Source/CameraCalibrationEditor/Private/Calibrators/CameraLensDistortionAlgoAruco.cpp
+++ b/Engine/Plugins/VirtualProduction/CameraCalibration/Source/CameraCalibrationEditor/Private/Calibrators/CameraLensDistortionAlgoAruco.cpp
@@ -264,30 +264,30 @@ bool UCameraLensDistortionAlgoAruco::OnViewportClicked(const FGeometry& MyGeomet
 	return true;
 }
 
-FDistortionCalibrationTask UCameraLensDistortionAlgoAruco::BeginCalibration()
+FDistortionCalibrationTask UCameraLensDistortionAlgoAruco::BeginCalibration(FText& OutErrorMessage)
 {
 	FDistortionCalibrationTask CalibrationTask = {};
 
 	if (CalibrationRows.Num() < 1)
 	{
-		UE_LOG(LogCameraCalibrationEditor, Error, TEXT("Could not initiate distortion calibration. At least 1 calibration row is required."));
+		OutErrorMessage = LOCTEXT("NotEnoughCalibrationRowsError", "Could not initiate distortion calibration. At least 1 calibration row is required.");
 		return CalibrationTask;
 	}
 
 	ULensDistortionTool* LensDistortionTool = WeakTool.Get();
-	if (!LensDistortionTool)
+	if (!ensureMsgf((LensDistortionTool != nullptr), TEXT("The Lens Distortion Tool was invalid.")))
 	{
 		return CalibrationTask;
 	}
 
 	FCameraCalibrationStepsController* StepsController = LensDistortionTool->GetCameraCalibrationStepsController();
-	if (!StepsController)
+	if (!ensureMsgf((StepsController != nullptr), TEXT("The Calibration Steps Controller was invalid.")))
 	{
 		return CalibrationTask;
 	}
 
 	ULensFile* LensFile = StepsController->GetLensFile();
-	if (!LensFile)
+	if (!ensureMsgf((LensFile != nullptr), TEXT("The Lens File was invalid.")))
 	{
 		return CalibrationTask;
 	}
@@ -304,7 +304,7 @@ FDistortionCalibrationTask UCameraLensDistortionAlgoAruco::BeginCalibration()
 
 	if (FMath::IsNearlyZero(PixelAspect))
 	{
-		UE_LOG(LogCameraCalibrationEditor, Error, TEXT("Could not initiate distortion calibration. The pixel aspect ratio of the CineCamera is zero, which is invalid."));
+		OutErrorMessage = LOCTEXT("PixelAspectZeroError", "The pixel aspect ratio of the CineCamera is zero, which is invalid.");
 		return CalibrationTask;
 	}
 
@@ -313,13 +313,13 @@ FDistortionCalibrationTask UCameraLensDistortionAlgoAruco::BeginCalibration()
 
 	if (FMath::IsNearlyZero(DesqueezeSensorWidth))
 	{
-		UE_LOG(LogCameraCalibrationEditor, Error, TEXT("Could not initiate distortion calibration. The sensor width of the CineCamera is zero, which is invalid."));
+		OutErrorMessage = LOCTEXT("SensorWidthZeroError", "One of the filmback dimensions of the CineCamera is zero, which is invalid.");
 		return CalibrationTask;
 	}
 
-	if (!FocalLengthEstimate.IsSet())
+	if (!FocalLengthEstimate.IsSet() || FMath::IsNearlyZero(FocalLengthEstimate.GetValue()))
 	{
-		UE_LOG(LogCameraCalibrationEditor, Error, TEXT("Could not initiate distortion calibration. Please enter a value (in mm) for the focal length estimate."));
+		OutErrorMessage = LOCTEXT("FocalLengthEstimateError", "Enter a non-zero value (in mm) for the focal length estimate.");
 		return CalibrationTask;
 	}
 
diff --git a/Engine/Plugins/VirtualProduction/CameraCalibration/Source/CameraCalibrationEditor/Private/Calibrators/CameraLensDistortionAlgoAruco.h b/Engine/Plugins/VirtualProduction/CameraCalibration/Source/CameraCalibrationEditor/Private/Calibrators/CameraLensDistortionAlgoAruco.h
--- a/Engine/Plugins/VirtualProduction/CameraCalibration/Source/CameraCalibrationEditor/Private/Calibrators/CameraLensDistortionAlgoAruco.h
+++ b/Engine/Plugins/VirtualProduction/CameraCalibration/Source/CameraCalibrationEditor/Private/Calibrators/CameraLensDistortionAlgoAruco.h
@@ -56,7 +56,7 @@ public:
 	virtual FName FriendlyName() const override { return TEXT("Lens Distortion Aruco"); };
 	virtual FName ShortName() const override { return TEXT("Aruco"); };
 	virtual void OnDistortionSavedToLens() override;
-	virtual FDistortionCalibrationTask BeginCalibration() override;
+	virtual FDistortionCalibrationTask BeginCalibration(FText& OutErrorMessage) override;
 	virtual bool SupportsAsyncCalibration() override { return true; };
 	virtual bool HasCalibrationData() const override;
 	virtual void PreImportCalibrationData() override;
diff --git a/Engine/Plugins/VirtualProduction/CameraCalibration/Source/CameraCalibrationEditor/Private/Calibrators/CameraLensDistortionAlgoCheckerboard.cpp b/Engine/Plugins/VirtualProduction/CameraCalibration/Source/CameraCalibrationEditor/Private/Calibrators/CameraLensDistortionAlgoCheckerboard.cpp
--- a/Engine/Plugins/VirtualProduction/CameraCalibration/Source/CameraCalibrationEditor/Private/Calibrators/CameraLensDistortionAlgoCheckerboard.cpp
+++ b/Engine/Plugins/VirtualProduction/CameraCalibration/Source/CameraCalibrationEditor/Private/Calibrators/CameraLensDistortionAlgoCheckerboard.cpp
@@ -609,14 +609,14 @@ bool UCameraLensDistortionAlgoCheckerboard::ValidateNewRow(TSharedPtr<FLensDisto
 	return true;
 }
 
-FDistortionCalibrationTask UCameraLensDistortionAlgoCheckerboard::BeginCalibration()
+FDistortionCalibrationTask UCameraLensDistortionAlgoCheckerboard::BeginCalibration(FText& OutErrorMessage)
 {
 	FDistortionCalibrationTask CalibrationTask = {};
 
 	// Validate that enough points were gathered to attempt a calibration
 	if (CalibrationRows.Num() < 1)
 	{
-		UE_LOG(LogCameraCalibrationEditor, Error, TEXT("Could not initiate distortion calibration. At least 1 calibration row is required."));
+		OutErrorMessage = LOCTEXT("NotEnoughCalibrationRowsError", "Could not initiate distortion calibration. At least 1 calibration row is required.");
 		return CalibrationTask;
 	}
 
@@ -630,19 +630,19 @@ FDistortionCalibrationTask UCameraLensDistortionAlgoCheckerboard::BeginCalibrati
 	}
 
 	ULensDistortionTool* LensDistortionTool = Tool.Get();
-	if (!LensDistortionTool)
+	if (!ensureMsgf((LensDistortionTool != nullptr), TEXT("The Lens Distortion Tool was invalid.")))
 	{
 		return CalibrationTask;
 	}
 
 	FCameraCalibrationStepsController* StepsController = LensDistortionTool->GetCameraCalibrationStepsController();
-	if (!StepsController)
+	if (!ensureMsgf((StepsController != nullptr), TEXT("The Calibration Steps Controller was invalid.")))
 	{
 		return CalibrationTask;
 	}
 
 	ULensFile* LensFile = StepsController->GetLensFile();
-	if (!LensFile)
+	if (!ensureMsgf((LensFile != nullptr), TEXT("The Lens File was invalid.")))
 	{
 		return CalibrationTask;
 	}
@@ -657,7 +657,7 @@ FDistortionCalibrationTask UCameraLensDistortionAlgoCheckerboard::BeginCalibrati
 
 	if (FMath::IsNearlyZero(PixelAspect))
 	{
-		UE_LOG(LogCameraCalibrationEditor, Error, TEXT("Could not initiate distortion calibration. The pixel aspect ratio of the CineCamera is zero, which is invalid."));
+		OutErrorMessage = LOCTEXT("PixelAspectZeroError", "The pixel aspect ratio of the CineCamera is zero, which is invalid.");
 		return CalibrationTask;
 	}
 
@@ -666,11 +666,18 @@ FDistortionCalibrationTask UCameraLensDistortionAlgoCheckerboard::BeginCalibrati
 
 	if (FMath::IsNearlyZero(DesqueezeSensorWidth))
 	{
-		UE_LOG(LogCameraCalibrationEditor, Error, TEXT("Could not initiate distortion calibration. The sensor width of the CineCamera is zero, which is invalid."));
+		OutErrorMessage = LOCTEXT("SensorWidthZeroError", "One of the filmback dimensions of the CineCamera is zero, which is invalid.");
+		return CalibrationTask;
+	}
+
+	if (!FocalLengthEstimate.IsSet() || FMath::IsNearlyZero(FocalLengthEstimate.GetValue()))
+	{
+		OutErrorMessage = LOCTEXT("FocalLengthEstimateError", "Enter a non-zero value (in mm) for the focal length estimate.");
 		return CalibrationTask;
 	}
 
-	const double Fx = (FocalLengthEstimate / DesqueezeSensorWidth) * ImageSize.X;
+	const float FocalLengthEstimateValue = FocalLengthEstimate.GetValue();
+	const double Fx = (FocalLengthEstimateValue / DesqueezeSensorWidth) * ImageSize.X;
 
 	// When operating on a desqueezed image, we expect our pixel aspect to be square, so horizontal and vertical field of view are assumed to be equal (i.e. Fx == Fy)
 	FVector2D FocalLength = FVector2D(Fx, Fx);
@@ -885,9 +892,9 @@ TSharedRef<SWidget> UCameraLensDistortionAlgoCheckerboard::BuildFocalLengthEstim
 	.FillWidth(0.8)
 	.Padding(0, 0, 10, 0)
 	[
-		SNew(SNumericEntryBox<double>)
-		.Value(MakeAttributeLambda([this]() { return TOptional<double>(FocalLengthEstimate); }))
-		.OnValueChanged(SNumericEntryBox<double>::FOnValueChanged::CreateLambda([this](double NewValue) { FocalLengthEstimate = NewValue; }))
+		SNew(SNumericEntryBox<float>)
+		.Value_UObject(this, &UCameraLensDistortionAlgoCheckerboard::GetFocalLengthEstimate)
+		.OnValueChanged_UObject(this, &UCameraLensDistortionAlgoCheckerboard::SetFocalLengthEstimate)
 	]
 
 	+ SHorizontalBox::Slot()
@@ -895,21 +902,34 @@ TSharedRef<SWidget> UCameraLensDistortionAlgoCheckerboard::BuildFocalLengthEstim
 	[
 		SNew(SCheckBox)
 		.Padding(FMargin(5, 0, 15, 0))
-		.IsChecked_Lambda([&]() -> ECheckBoxState
-		{
-			return bFixFocalLength ? ECheckBoxState::Checked : ECheckBoxState::Unchecked;
-		})
-		.OnCheckStateChanged_Lambda([&](ECheckBoxState NewState) -> void
-		{
-			bFixFocalLength = (NewState == ECheckBoxState::Checked);
-		})
+		.IsChecked_UObject(this, &UCameraLensDistortionAlgoCheckerboard::IsFixFocalLengthChecked)
+		.OnCheckStateChanged_UObject(this, &UCameraLensDistortionAlgoCheckerboard::OnFixFocalLengthCheckStateChanged)
 		[
-			SNew(STextBlock)
-			.Text(LOCTEXT("FixText", "Fix?"))
+			SNew(STextBlock).Text(LOCTEXT("FixText", "Fix?"))
 		]
 	];
 }
 
+TOptional<float> UCameraLensDistortionAlgoCheckerboard::GetFocalLengthEstimate() const
+{
+	return FocalLengthEstimate;
+}
+
+void UCameraLensDistortionAlgoCheckerboard::SetFocalLengthEstimate(float NewValue)
+{
+	FocalLengthEstimate = NewValue;
+}
+
+ECheckBoxState UCameraLensDistortionAlgoCheckerboard::IsFixFocalLengthChecked() const
+{
+	return bFixFocalLength ? ECheckBoxState::Checked : ECheckBoxState::Unchecked;
+}
+
+void UCameraLensDistortionAlgoCheckerboard::OnFixFocalLengthCheckStateChanged(ECheckBoxState NewState)
+{
+	bFixFocalLength = (NewState == ECheckBoxState::Checked);
+}
+
 TSharedRef<SWidget> UCameraLensDistortionAlgoCheckerboard::BuildFixImageCenterWidget()
 {
 	return SNew(SCheckBox)
diff --git a/Engine/Plugins/VirtualProduction/CameraCalibration/Source/CameraCalibrationEditor/Private/Calibrators/CameraLensDistortionAlgoCheckerboard.h b/Engine/Plugins/VirtualProduction/CameraCalibration/Source/CameraCalibrationEditor/Private/Calibrators/CameraLensDistortionAlgoCheckerboard.h
--- a/Engine/Plugins/VirtualProduction/CameraCalibration/Source/CameraCalibrationEditor/Private/Calibrators/CameraLensDistortionAlgoCheckerboard.h
+++ b/Engine/Plugins/VirtualProduction/CameraCalibration/Source/CameraCalibrationEditor/Private/Calibrators/CameraLensDistortionAlgoCheckerboard.h
@@ -135,7 +135,7 @@ public:
 	virtual UMaterialInterface* GetOverlayMaterial() const override;
 	virtual bool IsOverlayEnabled() const override { return bShouldShowOverlay; };
 	virtual void OnDistortionSavedToLens() override;
-	virtual FDistortionCalibrationTask BeginCalibration() override;
+	virtual FDistortionCalibrationTask BeginCalibration(FText& OutErrorMessage) override;
 	virtual void CancelCalibration() override;
 	virtual bool GetCalibrationStatus(FText& StatusText) const override;
 	virtual bool SupportsAsyncCalibration() override { return true; };
@@ -182,8 +182,8 @@ private:
 	/** If true, the solver will not solve for image center, and will just use the input value */
 	bool bFixImageCenter = false;
 
-	/** Estimated focal length to initialize the distortion solver */
-	double FocalLengthEstimate = 0.0;
+	/** Estimate for the focal length to provide to the solver. If not set, user will be warned to set before continuing calibration */
+	TOptional<float> FocalLengthEstimate;
 
 	/** If true, the solver will use the current camera pose to initialize the camera extrinsic parameters for each image */
 	bool bUseExtrinsicsGuess = false;
@@ -226,6 +226,18 @@ private:
 	/** Builds the UI for the action buttons (RemoveLast, ClearAll) */
 	TSharedRef<SWidget> BuildCalibrationActionButtons();
 
+	/** Get the current value of the focal length estimate */
+	TOptional<float> GetFocalLengthEstimate() const;
+
+	/** Set the current value of the focal length estimate */
+	void SetFocalLengthEstimate(float NewValue);
+
+	/** Returns the state of the fix focal length checkbox */
+	ECheckBoxState IsFixFocalLengthChecked() const;
+
+	/** Changes the state of the fix focal length checkbox */
+	void OnFixFocalLengthCheckStateChanged(ECheckBoxState NewState);
+
 private:
 
 	/** Returns the first checkerboard object in the scene that it can find */
diff --git a/Engine/Plugins/VirtualProduction/CameraCalibrationCore/Source/CameraCalibrationCore/Public/CameraLensDistortionAlgo.h b/Engine/Plugins/VirtualProduction/CameraCalibrationCore/Source/CameraCalibrationCore/Public/CameraLensDistortionAlgo.h
--- a/Engine/Plugins/VirtualProduction/CameraCalibrationCore/Source/CameraCalibrationCore/Public/CameraLensDistortionAlgo.h
+++ b/Engine/Plugins/VirtualProduction/CameraCalibrationCore/Source/CameraCalibrationCore/Public/CameraLensDistortionAlgo.h
@@ -86,7 +86,7 @@ public:
 	};
 
 	/** Launches an asynchronous task to perform the distortion calibration. The task will contain the calibration result when it finishes. */
-	virtual FDistortionCalibrationTask BeginCalibration() { return {}; };
+	virtual FDistortionCalibrationTask BeginCalibration(FText& OutErrorMessage) { return {}; };
 
 	/** Cancel an in-progress async calibration task */
 	virtual void CancelCalibration() { };
