CameraCalibration: Fix warning in output log about improperly initialized Ustruct member.
#rnx

diff --git a/Engine/Plugins/VirtualProduction/CameraCalibration/Source/CameraCalibrationEditor/Private/CameraCalibrationUtilsPrivate.h b/Engine/Plugins/VirtualProduction/CameraCalibration/Source/CameraCalibrationEditor/Private/CameraCalibrationUtilsPrivate.h
--- a/Engine/Plugins/VirtualProduction/CameraCalibration/Source/CameraCalibrationEditor/Private/CameraCalibrationUtilsPrivate.h
+++ b/Engine/Plugins/VirtualProduction/CameraCalibration/Source/CameraCalibrationEditor/Private/CameraCalibrationUtilsPrivate.h
@@ -15,11 +15,11 @@ struct FArucoCalibrationPoint
 
 	/** 3D locations (in world space) of each corner */
 	UPROPERTY()
-	FVector Corners3D[4];
+	FVector Corners3D[4] = { FVector::ZeroVector, FVector::ZeroVector, FVector::ZeroVector, FVector::ZeroVector };
 
 	/** 2D locations (in pixels) of each corner */
 	UPROPERTY()
-	FVector2f Corners2D[4];
+	FVector2f Corners2D[4] = { FVector2f::ZeroVector, FVector2f::ZeroVector, FVector2f::ZeroVector, FVector2f::ZeroVector };
 
 	/** Aruco Marker ID */
 	UPROPERTY()
