Fix unreachable code issue (presumably a typo).

#rnx

#changelist validated
#virtualized

diff --git a/Engine/Plugins/VirtualProduction/CameraCalibration/Source/CameraCalibrationEditor/Private/UI/CameraCalibrationMenuEntry.cpp b/Engine/Plugins/VirtualProduction/CameraCalibration/Source/CameraCalibrationEditor/Private/UI/CameraCalibrationMenuEntry.cpp
--- a/Engine/Plugins/VirtualProduction/CameraCalibration/Source/CameraCalibrationEditor/Private/UI/CameraCalibrationMenuEntry.cpp
+++ b/Engine/Plugins/VirtualProduction/CameraCalibration/Source/CameraCalibrationEditor/Private/UI/CameraCalibrationMenuEntry.cpp
@@ -106,7 +106,7 @@ struct FCameraCalibrationMenuEntryImpl
 			//and current default engine lens file.
 			GetMutableDefault<UCameraCalibrationEditorSettings>()->SetUserLensFile(NewAsset);
 			UCameraCalibrationSubsystem* SubSystem = GEngine->GetEngineSubsystem<UCameraCalibrationSubsystem>();
-			return SubSystem->SetDefaultLensFile(NewAsset);
+			SubSystem->SetDefaultLensFile(NewAsset);
 
 			GEditor->GetEditorSubsystem<UAssetEditorSubsystem>()->OpenEditorForAsset(NewAsset);
 		}
@@ -125,7 +125,7 @@ struct FCameraCalibrationMenuEntryImpl
 		//and current default engine lens file.
 		GetMutableDefault<UCameraCalibrationEditorSettings>()->SetUserLensFile(Asset);
 		UCameraCalibrationSubsystem* SubSystem = GEngine->GetEngineSubsystem<UCameraCalibrationSubsystem>();
-		return SubSystem->SetDefaultLensFile(Asset);
+		SubSystem->SetDefaultLensFile(Asset);
 	}
 
 	TSharedRef<SWidget> GenerateMenuContent()
