CameraCalibration: Remove nerfacc python dependency from camera calibration plugin
#rnx

diff --git a/Engine/Plugins/VirtualProduction/CameraCalibration/CameraCalibration.uplugin b/Engine/Plugins/VirtualProduction/CameraCalibration/CameraCalibration.uplugin
--- a/Engine/Plugins/VirtualProduction/CameraCalibration/CameraCalibration.uplugin
+++ b/Engine/Plugins/VirtualProduction/CameraCalibration/CameraCalibration.uplugin
@@ -45,13 +45,5 @@
 			"Name": "PythonFoundationPackages",
 			"Enabled": true
 		}
-	],
-	"PythonRequirements": [
-		{
-			"Platform": "All",
-			"Requirements": [
-				"nerfacc==0.5.3 --hash=sha256:bcf9ecc85f2dc3439491f11f6103de9d5a0dfd2e9fe6ba4cade608443d7c6d2f"
-			]
-		}
 	]
 }
\ No newline at end of file
