Camera Calibration: Added the ability to change focus and zoom values on calibrated data points, with UI to handle merging/replacing existing points, as well as editing multiple linked values at the same time

diff --git a/Engine/Plugins/VirtualProduction/CameraCalibration/Source/CameraCalibrationEditor/Private/AssetEditor/SCameraCalibrationLinkedPointsDialog.cpp b/Engine/Plugins/VirtualProduction/CameraCalibration/Source/CameraCalibrationEditor/Private/AssetEditor/SCameraCalibrationLinkedPointsDialog.cpp
new file mode 100644
--- /dev/null
+++ b/Engine/Plugins/VirtualProduction/CameraCalibration/Source/CameraCalibrationEditor/Private/AssetEditor/SCameraCalibrationLinkedPointsDialog.cpp
@@ -0,0 +1,371 @@
+﻿// Copyright Epic Games, Inc. All Rights Reserved.
+
+#include "SCameraCalibrationLinkedPointsDialog.h"
+
+#include "Editor.h"
+#include "LensFile.h"
+#include "Widgets/SWindow.h"
+#include "Widgets/Input/SButton.h"
+#include "Widgets/Input/SCheckBox.h"
+
+#define LOCTEXT_NAMESPACE "SCameraCalibrationLinkedPointsDialog"
+
+class FLinkedTreeItem : public TSharedFromThis<FLinkedTreeItem>
+{
+protected:
+	typedef SCameraCalibrationLinkedPointsDialog::FLinkedItem FLinkedItem;
+	
+public:
+	FLinkedTreeItem(const FLinkedItem& InItem, bool bInIsSelected)
+		: Item(InItem)
+		, bIsSelected(bInIsSelected)
+	{ }
+	
+	virtual ~FLinkedTreeItem() = default;
+	
+	/** Generate the row widget */
+	virtual TSharedRef<ITableRow> MakeTreeRowWidget(const TSharedRef<STableViewBase>& InOwnerTable) = 0;
+	
+	/** Get children of this item */
+	const TArray<TSharedPtr<FLinkedTreeItem>>& GetChildren() const { return Children; }
+
+	void AddChild(const TSharedPtr<FLinkedTreeItem>& InChild) { Children.Add(InChild); }
+	
+	/** Gets whether this item is selected in the tree view */
+	bool IsSelected() const { return bIsSelected; }
+
+	/** Gets the linked item this tree item represents */
+	const FLinkedItem& GetItem() const { return Item; }
+
+	/** If this tree item is selected, add its linked item to the specified array, and recurse through its children */
+	void AddItemToSelectedLinkedItems(TArray<FLinkedItem>& SelectedLinkedItems)
+	{
+		for (const TSharedPtr<FLinkedTreeItem>& ChildItem : Children)
+		{
+			ChildItem->AddItemToSelectedLinkedItems(SelectedLinkedItems);
+		}
+
+		if (bIsSelected)
+		{
+			SelectedLinkedItems.Add(Item);
+		}
+	}
+
+protected:
+	ECheckBoxState GetIsSelectedCheckState() const { return bIsSelected ? ECheckBoxState::Checked : ECheckBoxState::Unchecked; }
+	virtual void OnCheckStateChanged(ECheckBoxState NewState) { bIsSelected = NewState == ECheckBoxState::Checked; }
+	
+protected:
+	/** The linked item this tree item represents */
+	FLinkedItem Item;
+
+	/** Parent of this item */
+	TWeakPtr<FLinkedTreeItem> Parent = nullptr;
+	
+	/** Children of this item */
+	TArray<TSharedPtr<FLinkedTreeItem>> Children;
+
+	/** Indicates if this item is selected in the tree view */
+	bool bIsSelected = false;
+
+	friend class FLinkedFocusTreeItem;
+};
+
+class FLinkedFocusTreeItem : public FLinkedTreeItem
+{
+public:
+	FLinkedFocusTreeItem(const FLinkedItem& InItem, bool bInIsSelected)
+		: FLinkedTreeItem(InItem, bInIsSelected)
+	{ }
+
+	virtual TSharedRef<ITableRow> MakeTreeRowWidget(const TSharedRef<STableViewBase>& InOwnerTable) override
+	{
+		const FName PointName = FBaseLensTable::GetFriendlyPointName(Item.Category);
+		const FText FocusLabel = FText::Format(LOCTEXT("FocusLabel", "{0}. Focus: {1}"), FText::FromName(PointName), Item.Focus);
+		
+		return SNew(STableRow<TSharedPtr<FLinkedTreeItem>>, InOwnerTable)
+		.Content()
+		[
+			SNew(SHorizontalBox)
+			+ SHorizontalBox::Slot()
+			.VAlign(VAlign_Center)
+			.AutoWidth()
+			[
+				SNew(SCheckBox)
+				.IsChecked(this, &FLinkedFocusTreeItem::GetIsSelectedCheckState)
+				.OnCheckStateChanged(this, &FLinkedFocusTreeItem::OnCheckStateChanged)
+			]
+			
+			+SHorizontalBox::Slot()
+			.Padding(5.0f, 5.0f)
+			.VAlign(VAlign_Center)
+			.HAlign(HAlign_Left)
+			[
+				SNew(STextBlock)
+				.Text(FocusLabel)
+			]
+		];
+	}
+
+protected:
+	virtual void OnCheckStateChanged(ECheckBoxState NewState) override
+	{
+		FLinkedTreeItem::OnCheckStateChanged(NewState);
+
+		// When a focus item is selected, all of its child zoom items are also selected
+		for (const TSharedPtr<FLinkedTreeItem>& Child : Children)
+		{
+			Child->bIsSelected = bIsSelected;
+		}
+	}
+};
+
+class FLinkedZoomTreeItem : public FLinkedTreeItem
+{
+public:
+	FLinkedZoomTreeItem(const FLinkedItem& InItem, bool bInIsSelected, const TSharedPtr<FLinkedTreeItem>& InParent)
+		: FLinkedTreeItem(InItem, bInIsSelected)
+	{
+		Parent = InParent;
+	} 
+
+	virtual TSharedRef<ITableRow> MakeTreeRowWidget(const TSharedRef<STableViewBase>& InOwnerTable) override
+	{
+		const FName PointName = FBaseLensTable::GetFriendlyPointName(Item.Category);
+		const FText ZoomLabel = Parent.IsValid() ?
+			FText::Format(LOCTEXT("ZoomLabel", "Zoom: {0}"), Item.Zoom.GetValue()) :
+			FText::Format(LOCTEXT("StandaloneZoomLabel", "{0}. Focus: {1}, Zoom: {2}"), FText::FromName(PointName), Item.Focus, Item.Zoom.GetValue());
+		
+		return SNew(STableRow<TSharedPtr<FLinkedTreeItem>>, InOwnerTable)
+		.Content()
+		[
+			SNew(SHorizontalBox)
+			+ SHorizontalBox::Slot()
+			.VAlign(VAlign_Center)
+			.AutoWidth()
+			[
+				SNew(SCheckBox)
+				.IsChecked(this, &FLinkedZoomTreeItem::GetIsSelectedCheckState)
+				.OnCheckStateChanged(this, &FLinkedZoomTreeItem::OnCheckStateChanged)
+			]
+			
+			+SHorizontalBox::Slot()
+			.Padding(5.0f, 5.0f)
+			.VAlign(VAlign_Center)
+			.HAlign(HAlign_Left)
+			[
+				SNew(STextBlock)
+				.Text(ZoomLabel)
+			]
+		];
+	}
+};
+
+void SCameraCalibrationLinkedPointsDialog::Construct(const FArguments& InArgs , ULensFile* InLensFile, const FLinkedItem& InInitialItem)
+{
+	WeakLensFile = InLensFile;
+	InitialItem = InInitialItem;
+	LinkedItemMode = InArgs._LinkedItemMode;
+	OnApplyLinkedAction = InArgs._OnApplyLinkedAction;
+
+	ChildSlot
+	[
+		SNew(SVerticalBox)
+		
+        +SVerticalBox::Slot()
+        .Padding(5.0f, 5.0f)
+        .AutoHeight()
+        [
+        	SNew(SBorder)
+        	.HAlign(HAlign_Center)
+        	.VAlign(VAlign_Center)
+        	.Padding(FMargin(4.f, 4.f, 4.f, 4.f))
+        	[
+        		SNew(STextBlock).Text(InArgs._DialogText)
+        	]
+        ]
+        
+        +SVerticalBox::Slot()
+        .Padding(5.0f, 5.0f)
+        .FillHeight(1.f)
+        [
+        	SAssignNew(LinkedItemsTree, STreeView<TSharedPtr<FLinkedTreeItem>>)
+			.TreeItemsSource(&LinkedItems)
+			.ItemHeight(24.0f)
+			.OnGenerateRow_Lambda([](TSharedPtr<FLinkedTreeItem> InItem, const TSharedRef<STableViewBase>& OwnerTable) { return InItem->MakeTreeRowWidget(OwnerTable); })
+			.OnGetChildren_Lambda([](TSharedPtr<FLinkedTreeItem> InItem, TArray<TSharedPtr<FLinkedTreeItem>>& OutNodes) { OutNodes = InItem->GetChildren(); })
+			.ClearSelectionOnClick(false)
+        ]
+
+        +SVerticalBox::Slot()
+        .Padding(5.0f, 5.0f)
+        .AutoHeight()
+        [
+        	InArgs._Content.Widget
+        ]
+        
+        +SVerticalBox::Slot()
+        .Padding(5.0f, 5.0f)
+        .AutoHeight()
+        [
+        	SNew(SBorder)
+        	.HAlign(HAlign_Fill)
+        	.VAlign(VAlign_Fill)
+        	.Padding(FMargin(4.f, 4.f, 4.f, 4.f))
+        	[
+        		SNew(SHorizontalBox)
+        		
+				+SHorizontalBox::Slot()
+				[
+					SNew(SButton)
+					.ContentPadding(FAppStyle::GetMargin("StandardDialog.ContentPadding"))
+					.OnClicked(this, &SCameraCalibrationLinkedPointsDialog::OnAcceptButtonClicked)
+					.HAlign(HAlign_Center)
+					.Text(InArgs._AcceptButtonText)
+				]
+				
+				+SHorizontalBox::Slot()
+				[
+					SNew(SButton)
+					.ContentPadding(FAppStyle::GetMargin("StandardDialog.ContentPadding"))
+					.OnClicked(this, &SCameraCalibrationLinkedPointsDialog::OnCancelButtonClicked)
+					.HAlign(HAlign_Center)
+					.Text(LOCTEXT("CancelLabel", "Cancel"))
+				]
+        	]
+        ]
+	];
+	
+	UpdateLinkedItems();
+}
+
+void SCameraCalibrationLinkedPointsDialog::OpenWindow(const FText& WindowTitle, const TSharedRef<SCameraCalibrationLinkedPointsDialog>& DialogBox)
+{
+	TSharedRef<SWindow> ModalWindow = SNew(SWindow)
+		.Title(WindowTitle)
+		.ClientSize(FVector2D(570, 420));
+
+	DialogBox->WindowWeakPtr = ModalWindow;
+	ModalWindow->SetContent(DialogBox);
+
+	GEditor->EditorAddModalWindow(ModalWindow);
+}
+
+void SCameraCalibrationLinkedPointsDialog::UpdateLinkedItems()
+{
+	LinkedItems.Empty();
+
+	if (!WeakLensFile.IsValid())
+	{
+		LinkedItemsTree->RequestTreeRefresh();
+		return;
+	}
+	
+	const FBaseLensTable* const BaseDataTable = WeakLensFile.Get()->GetDataTable(InitialItem.Category);
+	if (!ensure(BaseDataTable))
+	{
+		LinkedItemsTree->RequestTreeRefresh();
+		return;	
+	}
+	
+	if (LinkedItemMode == ELinkedItemMode::Zoom && InitialItem.Zoom.IsSet())
+	{
+		// Only display the linked zoom values and not the focus values
+
+		// Add the initial item to the tree view
+		constexpr bool bIsSelected = true;
+		LinkedItems.Add(MakeShared<FLinkedZoomTreeItem>(InitialItem, bIsSelected, nullptr));
+
+		BaseDataTable->ForEachLinkedFocusPoint([this, InitialZoom = InitialItem.Zoom.GetValue()](const FBaseFocusPoint& InFocusPoint, ELensDataCategory InCategory, FLinkPointMetadata LinkPointMeta)
+		{
+			for (int32 Index = 0; Index < InFocusPoint.GetNumPoints(); ++Index)
+			{
+				const float Zoom = InFocusPoint.GetZoom(Index);
+				if (!FMath::IsNearlyEqual(Zoom ,InitialZoom, WeakLensFile->InputTolerance))
+				{
+					continue;
+				}
+				
+				FLinkedItem LinkedZoomItem(InCategory, InFocusPoint.GetFocus(), Zoom);
+				LinkedItems.Add(MakeShared<FLinkedZoomTreeItem>(LinkedZoomItem, LinkPointMeta.bRemoveByDefault, nullptr));
+			}
+		}, InitialItem.Focus, WeakLensFile->InputTolerance);
+	}
+	else
+	{
+		// Add the initial item to the tree view
+		BaseDataTable->ForEachFocusPoint([this, Category = InitialItem.Category](const FBaseFocusPoint& InFocusPoint)
+		{
+			//Add entry for focus
+			FLinkedItem FocusLinkedItem(Category, InFocusPoint.GetFocus());
+			
+			constexpr bool bIsSelected = true;
+			const TSharedPtr<FLinkedFocusTreeItem> FocusTreeItem = MakeShared<FLinkedFocusTreeItem>(FocusLinkedItem, bIsSelected);
+			LinkedItems.Add(FocusTreeItem);
+
+			LinkedItemsTree->SetItemExpansion(FocusTreeItem, true);
+
+			if (LinkedItemMode == ELinkedItemMode::Both)
+			{
+				for (int32 Index = 0; Index < InFocusPoint.GetNumPoints(); ++Index)
+				{
+					//Add zoom points for this focus
+					FLinkedItem ZoomLinkedItem(Category, InFocusPoint.GetFocus(), InFocusPoint.GetZoom(Index));
+					FocusTreeItem->AddChild(MakeShared<FLinkedZoomTreeItem>(ZoomLinkedItem, bIsSelected, FocusTreeItem));
+				}
+			}
+		}, InitialItem.Focus);
+
+		BaseDataTable->ForEachLinkedFocusPoint([this](const FBaseFocusPoint& InFocusPoint, ELensDataCategory InCategory, FLinkPointMetadata LinkPointMeta)
+		{
+			FLinkedItem FocusLinkedItem(InCategory, InFocusPoint.GetFocus());
+			const TSharedPtr<FLinkedFocusTreeItem> FocusTreeItem = MakeShared<FLinkedFocusTreeItem>(FocusLinkedItem, LinkPointMeta.bRemoveByDefault);
+			
+			LinkedItems.Add(FocusTreeItem);
+			LinkedItemsTree->SetItemExpansion(FocusTreeItem, true);
+
+			if (LinkedItemMode == ELinkedItemMode::Both)
+			{
+				for(int32 Index = 0; Index < InFocusPoint.GetNumPoints(); ++Index)
+				{
+					//Add zoom points for this focus
+					FLinkedItem ZoomLinkedItem(InCategory, InFocusPoint.GetFocus(), InFocusPoint.GetZoom(Index));
+					FocusTreeItem->AddChild(MakeShared<FLinkedZoomTreeItem>(ZoomLinkedItem, LinkPointMeta.bRemoveByDefault, FocusTreeItem));
+				}
+			}
+		}, InitialItem.Focus, WeakLensFile->InputTolerance);
+	}
+
+	LinkedItemsTree->RequestTreeRefresh();
+}
+
+FReply SCameraCalibrationLinkedPointsDialog::OnAcceptButtonClicked()
+{
+	TArray<FLinkedItem> SelectedLinkedItems;
+
+	for (const TSharedPtr<FLinkedTreeItem>& TreeItem : LinkedItems)
+	{
+		TreeItem->AddItemToSelectedLinkedItems(SelectedLinkedItems);
+	}
+
+	OnApplyLinkedAction.ExecuteIfBound(SelectedLinkedItems);
+
+	if (WindowWeakPtr.IsValid())
+	{
+		WindowWeakPtr.Pin()->RequestDestroyWindow();
+	}
+	
+	return FReply::Handled();
+}
+
+FReply SCameraCalibrationLinkedPointsDialog::OnCancelButtonClicked()
+{
+	if (WindowWeakPtr.IsValid())
+	{
+		WindowWeakPtr.Pin()->RequestDestroyWindow();
+	}
+	
+	return FReply::Handled();
+}
+
+#undef LOCTEXT_NAMESPACE
diff --git a/Engine/Plugins/VirtualProduction/CameraCalibration/Source/CameraCalibrationEditor/Private/AssetEditor/SCameraCalibrationLinkedPointsDialog.h b/Engine/Plugins/VirtualProduction/CameraCalibration/Source/CameraCalibrationEditor/Private/AssetEditor/SCameraCalibrationLinkedPointsDialog.h
new file mode 100644
--- /dev/null
+++ b/Engine/Plugins/VirtualProduction/CameraCalibration/Source/CameraCalibrationEditor/Private/AssetEditor/SCameraCalibrationLinkedPointsDialog.h
@@ -0,0 +1,106 @@
+﻿// Copyright Epic Games, Inc. All Rights Reserved.
+
+#pragma once
+
+#include "CoreMinimal.h"
+#include "LensFile.h"
+#include "Widgets/SCompoundWidget.h"
+#include "Widgets/Views/STreeView.h"
+
+class FLinkedTreeItem;
+class ULensFile;
+enum class ELensDataCategory : uint8;
+
+/**
+ * A generic dialog box that displays all points linked to a specific point, allowing the user to
+ * apply a single operation to multiple linked points
+ */
+class SCameraCalibrationLinkedPointsDialog : public SCompoundWidget
+{
+public:
+	/** An item within a lens data table that is linked to the current item */
+	struct FLinkedItem
+	{
+		FLinkedItem(ELensDataCategory InCategory, float InFocus)
+			: Category(InCategory)
+			, Focus(InFocus)
+		{ }
+
+		FLinkedItem(ELensDataCategory InCategory, float InFocus, float InZoom)
+			: Category(InCategory)
+			, Focus(InFocus)
+			, Zoom(InZoom)
+		{ }
+		
+		ELensDataCategory Category = (ELensDataCategory)0;
+		float Focus = 0.0f;
+		TOptional<float> Zoom = TOptional<float>();
+	};
+
+	/** Indicates the mode of the dialog box */
+	enum class ELinkedItemMode : uint8
+	{
+		/** Only linked focus items will be displayed and selected from */
+		Focus,
+
+		/** Only linked zoom items will be displayed and selected from */
+		Zoom,
+
+		/** Both focus and zoom items will be displayed, and their selection state is coupled */
+		Both
+	};
+	
+	DECLARE_DELEGATE_OneParam(FOnApplyLinkedAction, const TArray<FLinkedItem>& /**ItemsToApplyActionTo*/);
+	
+	SLATE_BEGIN_ARGS(SCameraCalibrationLinkedPointsDialog)
+		: _AcceptButtonText(NSLOCTEXT("SCameraCalibrationLinkedPointsDialog", "AcceptLabel", "Accept"))
+	{ }
+		SLATE_ARGUMENT(ELinkedItemMode, LinkedItemMode)
+		SLATE_ATTRIBUTE(FText, DialogText)
+		SLATE_ATTRIBUTE(FText, AcceptButtonText)
+		SLATE_EVENT(FOnApplyLinkedAction, OnApplyLinkedAction)
+		SLATE_DEFAULT_SLOT(FArguments, Content)
+	SLATE_END_ARGS()
+
+	/** Constructs this widget with InArgs */
+	void Construct(const FArguments& InArgs, ULensFile* InLensFile, const FLinkedItem& InInitialItem);
+
+	/**
+	 * Open a new linked points window and block until the window is closed
+	 * @param WindowTitle The title of the window
+	 * @param DialogBox The dialog widget to display in the window
+	 */
+	static void OpenWindow(const FText& WindowTitle, const TSharedRef<SCameraCalibrationLinkedPointsDialog>& DialogBox);
+
+private:
+	/** Rebuilds the list of linked items and refreshes the tree view */
+	void UpdateLinkedItems();
+	
+	/** Raised when the accept button is clicked */
+	FReply OnAcceptButtonClicked();
+
+	/** Raised when the cancel button is clicked */
+	FReply OnCancelButtonClicked();
+	
+private:
+	/** Modal window pointer */
+	TWeakPtr<SWindow> WindowWeakPtr;
+
+	/** LensFile we're editing */
+	TWeakObjectPtr<ULensFile> WeakLensFile = nullptr;
+
+	/** The item linked items are being found for */
+	FLinkedItem InitialItem = FLinkedItem((ELensDataCategory)0, 0.0f);
+	
+	/** A list of items that are linked */
+	TArray<TSharedPtr<FLinkedTreeItem>> LinkedItems;
+
+	/** Tree widget to display the linked items */
+	TSharedPtr<STreeView<TSharedPtr<FLinkedTreeItem>>> LinkedItemsTree;
+
+	/** The mode of the dialog box */
+	ELinkedItemMode LinkedItemMode = ELinkedItemMode::Both;
+
+	/** Delegate to raise when the accept button is pressed and an action needs to be applied to each selected linked item */
+	FOnApplyLinkedAction OnApplyLinkedAction;
+};
diff --git a/Engine/Plugins/VirtualProduction/CameraCalibration/Source/CameraCalibrationEditor/Private/AssetEditor/SCameraCalibrationRemovePointDialog.cpp b/Engine/Plugins/VirtualProduction/CameraCalibration/Source/CameraCalibrationEditor/Private/AssetEditor/SCameraCalibrationRemovePointDialog.cpp
deleted file mode 100644
diff --git a/Engine/Plugins/VirtualProduction/CameraCalibration/Source/CameraCalibrationEditor/Private/AssetEditor/SCameraCalibrationRemovePointDialog.h b/Engine/Plugins/VirtualProduction/CameraCalibration/Source/CameraCalibrationEditor/Private/AssetEditor/SCameraCalibrationRemovePointDialog.h
deleted file mode 100644
diff --git a/Engine/Plugins/VirtualProduction/CameraCalibration/Source/CameraCalibrationEditor/Private/AssetEditor/SLensDataEditPointDialog.h b/Engine/Plugins/VirtualProduction/CameraCalibration/Source/CameraCalibrationEditor/Private/AssetEditor/SLensDataEditPointDialog.h
--- a/Engine/Plugins/VirtualProduction/CameraCalibration/Source/CameraCalibrationEditor/Private/AssetEditor/SLensDataEditPointDialog.h
+++ b/Engine/Plugins/VirtualProduction/CameraCalibration/Source/CameraCalibrationEditor/Private/AssetEditor/SLensDataEditPointDialog.h
@@ -10,10 +10,12 @@
 #include "Modules/ModuleManager.h"
 #include "PropertyEditorModule.h"
 #include "ScopedTransaction.h"
+#include "UI/CameraCalibrationWidgetHelpers.h"
 #include "UObject/StrongObjectPtr.h"
 #include "UObject/StructOnScope.h"
 #include "Widgets/SBoxPanel.h"
 #include "Widgets/Input/SButton.h"
+#include "Widgets/Input/SNumericEntryBox.h"
 #include "Widgets/Layout/SBorder.h"
 #include "Widgets/Text/STextBlock.h"
 
@@ -26,21 +28,32 @@
 template<typename TStructType>
 class SLensDataEditPointDialog : public SCompoundWidget
 {
+private:
+	struct FTrackingInputData
+	{
+		float InitialValue = 0.0;
+		float CurrentValue = 0.0;
+		bool bIsOverridden = false;
+	};
+	
 public:
 	/** Delegate for save struct on scope with specific struct type */
-	DECLARE_DELEGATE_OneParam(FOnSave, TSharedPtr<TStructOnScope<TStructType>> /*InStructToEdit */)
+	DECLARE_DELEGATE_ThreeParams(FOnSave, TSharedPtr<TStructOnScope<TStructType>> /*InStructToEdit */, TOptional<float> /*NewFocus*/, TOptional<float> /*NewZoom*/)
 	
 	SLATE_BEGIN_ARGS(SLensDataEditPointDialog)
 	{}
-	SLATE_EVENT(FOnSave, OnSave)
+		SLATE_EVENT(FOnSave, OnSave)
 	SLATE_END_ARGS()
 
-	void Construct(const FArguments& InArgs, ULensFile* InLensFile, TSharedPtr<TStructOnScope<TStructType>> InStructToEdit)
+	void Construct(const FArguments& InArgs, ULensFile* InLensFile, ELensDataCategory InCategory, TSharedPtr<TStructOnScope<TStructType>> InStructToEdit, float InFocus, float InZoom)
 	{
 		StructToEdit = InStructToEdit;
         LensFile = TStrongObjectPtr<ULensFile>(InLensFile);
+		Category = InCategory;
+		Focus.InitialValue = Focus.CurrentValue = InFocus;
+		Zoom.InitialValue = Zoom.CurrentValue = InZoom;
 		OnSaveDelegate = InArgs._OnSave;
-    
+		
         const TSharedPtr<SWidget> LensDataWidget = [this]()
         {
         	TSharedPtr<SWidget> WidgetToReturn;
@@ -86,11 +99,25 @@ public:
         			.Text(LOCTEXT("CancelEditDataPoint", "Cancel"))
         		];
         }();
-    
+		
         ChildSlot
         [
         	SNew(SVerticalBox)
-        	+ SVerticalBox::Slot()
+        	
+        	+SVerticalBox::Slot()
+			.Padding(5.0f, 5.0f)
+			.AutoHeight()
+			[
+				SNew(SBorder)
+				.HAlign(HAlign_Fill)
+				.VAlign(VAlign_Fill)
+				.Padding(FMargin(4.f, 4.f, 4.f, 4.f))
+				[
+					MakeTrackingDataWidget()
+				]
+			]
+			
+        	+SVerticalBox::Slot()
         	.Padding(5.0f, 5.0f)
         	.FillHeight(1.f)
         	[
@@ -102,7 +129,8 @@ public:
         			LensDataWidget.ToSharedRef()
         		]
         	]
-        	+ SVerticalBox::Slot()
+        	
+        	+SVerticalBox::Slot()
         	.Padding(5.0f, 5.0f)
         	.AutoHeight()
         	[
@@ -118,16 +146,101 @@ public:
 	}
 
 private:
+	TSharedRef<SWidget> MakeTrackingDataWidget()
+	{
+		TSharedPtr<SWidget> TrackingWidget = SNullWidget::NullWidget;
+
+		const auto MakeRowWidget = [this](const FString& Label, FTrackingInputData& Data) -> TSharedRef<SWidget>
+		{
+			return SNew(SHorizontalBox)
+				+ SHorizontalBox::Slot()
+				.AutoWidth()
+				.Padding(0.f, 8.f, 0.f, 8.f)
+				.VAlign(VAlign_Center)
+				[
+					SNew(SCheckBox)
+					.IsChecked_Lambda([&Data]() { return Data.bIsOverridden ? ECheckBoxState::Checked : ECheckBoxState::Unchecked; })
+					.OnCheckStateChanged(FOnCheckStateChanged::CreateLambda([&Data](ECheckBoxState NewState) { Data.bIsOverridden = NewState == ECheckBoxState::Checked; }))
+					.ToolTipText(LOCTEXT("TrackingDataCheckboxTooltip", "Check to override incoming tracking data for this point"))
+				]
+				+ SHorizontalBox::Slot()
+				.VAlign(VAlign_Center)
+				.Padding(4.f, 0.f, 0.f, 0.f)
+				[
+					SNew(STextBlock)
+					.Text(FText::FromString(Label))
+				]
+				+ SHorizontalBox::Slot()
+				.VAlign(VAlign_Center)
+				[
+					SNew(SNumericEntryBox<float>)
+					.IsEnabled_Lambda([&Data](){ return Data.bIsOverridden; })
+					.Value_Lambda([&Data]() { return TOptional<float>(Data.CurrentValue); })
+					.OnValueChanged(SNumericEntryBox<float>::FOnValueChanged::CreateLambda([&Data](float NewValue) { Data.CurrentValue = NewValue; }))
+				];
+		};
+
+		//Based on category, either have one tracking input or two
+		switch (Category)
+		{
+			case ELensDataCategory::Focus:
+			{
+				TrackingWidget = MakeRowWidget(TEXT("Input Focus"), Focus);
+				break;
+			}
+			case ELensDataCategory::Iris:
+			{
+				TrackingWidget = MakeRowWidget(TEXT("Input Iris"), Focus);
+				break;
+			}
+			case ELensDataCategory::Zoom:
+			case ELensDataCategory::Distortion:
+			case ELensDataCategory::ImageCenter:
+			case ELensDataCategory::NodalOffset:
+			case ELensDataCategory::STMap:
+			default:
+			{
+				TrackingWidget = 
+					SNew(SVerticalBox)
+					+ SVerticalBox::Slot()
+					[
+						MakeRowWidget(TEXT("Input Focus"), Focus)
+					]
+					+ SVerticalBox::Slot()
+					[
+						MakeRowWidget(TEXT("Input Zoom"), Zoom)
+					];
+				break;
+			}	
+		}
+
+		return TrackingWidget.ToSharedRef();
+	}
+	
 	/** Save Button handler */
 	FReply OnSaveDataPointClicked() const
 	{
-		if (StructToEdit.IsValid())
+		if (!StructToEdit.IsValid())
+		{
+			return FReply::Unhandled();
+		}
+
+		const FScopedTransaction MapPointEdited(LOCTEXT("MapPointEdited", "Map Point Edited"));
+		LensFile->Modify();
+		
+		TOptional<float> NewFocus = TOptional<float>();
+		if (Focus.bIsOverridden && Focus.CurrentValue != Focus.InitialValue)
 		{
-			const FScopedTransaction MapPointEdited(LOCTEXT("MapPointEdited", "Map Point Edited"));
-			LensFile->Modify();
+			NewFocus = Focus.CurrentValue;
 		}
 
-		OnSaveDelegate.ExecuteIfBound(StructToEdit);
+		TOptional<float> NewZoom = TOptional<float>();
+		if (Zoom.bIsOverridden && Zoom.CurrentValue != Zoom.InitialValue)
+		{
+			NewZoom = Zoom.CurrentValue;
+		}
+		
+		OnSaveDelegate.ExecuteIfBound(StructToEdit, NewFocus, NewZoom);
 
 		FCameraCalibrationToolkit::DestroyPopupWindow();
 	
@@ -146,9 +259,18 @@ private:
 	/** LensFile being edited */
 	TStrongObjectPtr<ULensFile> LensFile;
 
+	/** The category of the data struct being edited */
+	ELensDataCategory Category = ELensDataCategory::Focus;
+	
 	/** Editing Struct for visualize with StructureDetailsView */ 
 	TSharedPtr<TStructOnScope<TStructType>> StructToEdit;
+	
+	/** Focus of the data point being edited */
+	FTrackingInputData Focus;
 
+	/** Zoom of the data point being edited */
+	FTrackingInputData Zoom;
+	
 	/** On save struct delegate instance */
 	FOnSave OnSaveDelegate;
 };
@@ -171,23 +293,57 @@ namespace LensDataEditPointDialog
 
 		return ReturnStruct;
 	}
-
+	
 	/** Finds currently opened dialog window or spawns a new one for editing */
 	template<typename TStructType, typename TableType>
-	static void OpenDialog(ULensFile* InLensFile, const float InFocus, const float InZoom, TableType& InTable)
+	static void OpenDialog(ULensFile* InLensFile, ELensDataCategory InCategory, const float InFocus, const float InZoom, TableType& InTable, const FSimpleDelegate& OnPointSaved)
 	{
 		using FOnSaveType = typename SLensDataEditPointDialog<TStructType>::FOnSave;
 		using FStructType = TSharedPtr<TStructOnScope<TStructType>>;
 
 		// On save struct delegate
-		FOnSaveType OnSaveDelegate = FOnSaveType::CreateLambda([InLensFile, InFocus, InZoom, &InTable](FStructType InStructOnScope)
+		FOnSaveType OnSaveDelegate = FOnSaveType::CreateLambda([InFocus, InZoom, &InTable, OnPointSaved](FStructType InStructOnScope, TOptional<float> NewFocus, TOptional<float> NewZoom)
 		{
 			InTable.SetPoint(InFocus, InZoom, *InStructOnScope->Get());
+
+			// Set new zoom first so that any zoom replace conflicts are dealt with before attempting to change/merge the focus
+			if (NewZoom.IsSet())
+			{
+				if (InTable.HasZoomPoint(InFocus, NewZoom.GetValue()))
+				{
+					if (FCameraCalibrationWidgetHelpers::ShowReplaceZoomWarning())
+					{
+						InTable.ChangeZoomPoint(InFocus, InZoom, NewZoom.GetValue());
+					}
+				}
+				else
+				{
+					InTable.ChangeZoomPoint(InFocus, InZoom, NewZoom.GetValue());
+				}
+			}
+			
+			if (NewFocus.IsSet())
+			{
+				if (InTable.HasFocusPoint(NewFocus.GetValue()))
+				{
+					bool bReplaceExisting = false;
+					if (FCameraCalibrationWidgetHelpers::ShowMergeFocusWarning(bReplaceExisting))
+					{
+						InTable.MergeFocusPoint(InFocus, NewFocus.GetValue(), bReplaceExisting);
+					}
+				}
+				else
+				{
+					InTable.ChangeFocusPoint(InFocus, NewFocus.GetValue());
+				}
+			}
+
+			OnPointSaved.ExecuteIfBound();
 		});
 
 		FStructType InStructToEdit = GetStructToEdit<TStructType>(InFocus, InZoom, InTable);
 		TSharedPtr<SWindow> PopupWindow = FCameraCalibrationToolkit::OpenPopupWindow(LOCTEXT("LensEditorEditPointDialog", "Edit Lens Data Point"));
-		PopupWindow->SetContent(SNew(SLensDataEditPointDialog<TStructType>, InLensFile, InStructToEdit).OnSave(OnSaveDelegate));
+		PopupWindow->SetContent(SNew(SLensDataEditPointDialog<TStructType>, InLensFile, InCategory, InStructToEdit, InFocus, InZoom).OnSave(OnSaveDelegate));
 	}
 }
 
diff --git a/Engine/Plugins/VirtualProduction/CameraCalibration/Source/CameraCalibrationEditor/Private/AssetEditor/SLensDataListItem.cpp b/Engine/Plugins/VirtualProduction/CameraCalibration/Source/CameraCalibrationEditor/Private/AssetEditor/SLensDataListItem.cpp
--- a/Engine/Plugins/VirtualProduction/CameraCalibration/Source/CameraCalibrationEditor/Private/AssetEditor/SLensDataListItem.cpp
+++ b/Engine/Plugins/VirtualProduction/CameraCalibration/Source/CameraCalibrationEditor/Private/AssetEditor/SLensDataListItem.cpp
@@ -3,27 +3,30 @@
 #include "SLensDataListItem.h"
 
 #include "LensFile.h"
+#include "SCameraCalibrationLinkedPointsDialog.h"
 #include "ScopedTransaction.h"
 #include "SLensDataEditPointDialog.h"
 #include "UI/CameraCalibrationEditorStyle.h"
 #include "Widgets/Images/SImage.h"
 #include "Widgets/Input/SButton.h"
-#include "SCameraCalibrationRemovePointDialog.h"
+#include "UI/CameraCalibrationWidgetHelpers.h"
+#include "Widgets/Input/SNumericEntryBox.h"
+#include "Widgets/Layout/SWidgetSwitcher.h"
 #include "Widgets/Views/SListView.h"
 #include "Widgets/Views/STableRow.h"
 
 #define LOCTEXT_NAMESPACE "LensDataListItem"
 
-FLensDataListItem::FLensDataListItem(ULensFile* InLensFile, ELensDataCategory InCategory, int32 InSubCategoryIndex, FOnDataRemoved InOnDataRemovedCallback)
+FLensDataListItem::FLensDataListItem(ULensFile* InLensFile, ELensDataCategory InCategory, int32 InSubCategoryIndex, FOnDataChanged InOnDataChangedCallback)
 	: Category(InCategory)
 	, SubCategoryIndex(InSubCategoryIndex)
 	, WeakLensFile(InLensFile)
-	, OnDataRemovedCallback(InOnDataRemovedCallback)
+	, OnDataChangedCallback(InOnDataChangedCallback)
 {
 }
 
-FEncoderDataListItem::FEncoderDataListItem(ULensFile* InLensFile, ELensDataCategory InCategory, float InInputValue, int32 InIndex, FOnDataRemoved InOnDataRemovedCallback)
-	: FLensDataListItem(InLensFile, InCategory, INDEX_NONE, InOnDataRemovedCallback)
+FEncoderDataListItem::FEncoderDataListItem(ULensFile* InLensFile, ELensDataCategory InCategory, float InInputValue, int32 InIndex, FOnDataChanged InOnDataChangedCallback)
+	: FLensDataListItem(InLensFile, InCategory, INDEX_NONE, InOnDataChangedCallback)
 	, InputValue(InInputValue)
 	, EntryIndex()
 {
@@ -38,7 +41,7 @@ void FEncoderDataListItem::OnRemoveRequested() const
 
 		//Pass encoder mapping raw input value as focus to remove it
 		LensFilePtr->RemoveFocusPoint(Category, InputValue);
-		OnDataRemovedCallback.ExecuteIfBound(InputValue, TOptional<float>());
+		OnDataChangedCallback.ExecuteIfBound(ELensDataChangedReason::DataRemoved, InputValue, TOptional<float>());
 	}
 }
 
@@ -50,8 +53,8 @@ TSharedRef<ITableRow> FEncoderDataListItem::MakeTreeRowWidget(const TSharedRef<S
 		.AllowRemoval(true);
 }
 
-FFocusDataListItem::FFocusDataListItem(ULensFile* InLensFile, ELensDataCategory InCategory, int32 InSubCategoryIndex, float InFocus, FOnDataRemoved InOnDataRemovedCallback)
-	: FLensDataListItem(InLensFile, InCategory, InSubCategoryIndex, InOnDataRemovedCallback)
+FFocusDataListItem::FFocusDataListItem(ULensFile* InLensFile, ELensDataCategory InCategory, int32 InSubCategoryIndex, float InFocus, FOnDataChanged InOnDataChangedCallback)
+	: FLensDataListItem(InLensFile, InCategory, InSubCategoryIndex, InOnDataChangedCallback)
 	, Focus(InFocus)
 {
 }
@@ -61,6 +64,8 @@ TSharedRef<ITableRow> FFocusDataListItem::MakeTreeRowWidget(const TSharedRef<STa
 	return SNew(SLensDataItem, InOwnerTable, AsShared())
 		.EntryLabel(LOCTEXT("FocusLabel", "Focus: "))
 		.EntryValue(Focus)
+		.AllowEditEntryValue(true)
+		.OnEntryValueChanged(this, &FFocusDataListItem::OnFocusValueChanged)
 		.AllowRemoval(SubCategoryIndex == INDEX_NONE);
 }
 
@@ -76,14 +81,10 @@ void FFocusDataListItem::OnRemoveRequested() const
 		
 		if (LinkDataTable->HasLinkedFocusValues(Focus))
 		{
-			SCameraCalibrationRemovePointDialog::OpenWindow(
-				LensFilePtr,
-				Category,
-				FSimpleDelegate::CreateLambda([this]()
-				{
-					OnDataRemovedCallback.ExecuteIfBound(Focus, TOptional<float>());
-				}),
-				Focus);
+			if (RemoveLinkedFocusValues())
+			{
+				OnDataChangedCallback.ExecuteIfBound(ELensDataChangedReason::DataRemoved, Focus, TOptional<float>());
+			}
 		}
 		else
 		{
@@ -91,13 +92,166 @@ void FFocusDataListItem::OnRemoveRequested() const
 			LensFilePtr->Modify();
 
 			LensFilePtr->RemoveFocusPoint(Category, Focus);
-			OnDataRemovedCallback.ExecuteIfBound(Focus, TOptional<float>());
+			OnDataChangedCallback.ExecuteIfBound(ELensDataChangedReason::DataRemoved, Focus, TOptional<float>());
 		}
 	}
 }
 
-FZoomDataListItem::FZoomDataListItem(ULensFile* InLensFile, ELensDataCategory InCategory, int32 InSubCategoryIndex, const TSharedRef<FFocusDataListItem> InParent, float InZoom, FOnDataRemoved InOnDataRemovedCallback)
-	: FLensDataListItem(InLensFile, InCategory, InSubCategoryIndex, InOnDataRemovedCallback)
+bool FFocusDataListItem::OnFocusValueChanged(float NewFocusValue)
+{
+	if (ULensFile* LensFilePtr = WeakLensFile.Get())
+	{
+		const FBaseLensTable* const DataTable = LensFilePtr->GetDataTable(Category);
+		if (!ensure(DataTable))
+		{
+			return false;
+		}
+
+		if (DataTable->HasLinkedFocusValues(Focus))
+		{
+			// If there are any linked focus values, allow the linked focus dialog box to handle any needed changes to the focuses in the lens file
+			if (ChangeLinkedFocusValues(NewFocusValue))
+			{
+				Focus = NewFocusValue;
+				OnDataChangedCallback.ExecuteIfBound(ELensDataChangedReason::DataChanged, Focus, TOptional<float>());
+				return true;
+			}
+
+			return false;
+		}
+		
+		if (LensFilePtr->HasFocusPoint(Category, NewFocusValue))
+		{
+			// If the data already has a point for the new focus value, we want to merge this focus point's data with
+			// the existing focus points, with user permission
+			bool bReplaceExistingZoomPoints = false;
+			const bool bMergeFocusPoints = FCameraCalibrationWidgetHelpers::ShowMergeFocusWarning(bReplaceExistingZoomPoints);
+			if (bMergeFocusPoints)
+			{
+				FScopedTransaction Transaction(LOCTEXT("MergeFocusPointTransaction", "Merge Focus Point"));
+				LensFilePtr->Modify();
+
+				LensFilePtr->MergeFocusPoint(Category, Focus, NewFocusValue, bReplaceExistingZoomPoints);
+
+				Focus = NewFocusValue;
+				OnDataChangedCallback.ExecuteIfBound(ELensDataChangedReason::DataChanged, Focus, TOptional<float>());
+
+				return true;
+			}
+		}
+		else
+		{
+			// Otherwise, we can just change the data point's value directly
+			FScopedTransaction Transaction(LOCTEXT("ChangeFocusPointTransaction", "Change Focus Point"));
+			LensFilePtr->Modify();
+
+			LensFilePtr->ChangeFocusPoint(Category, Focus, NewFocusValue);
+			Focus = NewFocusValue;
+
+			return true;
+		}
+	}
+
+	return false;
+}
+
+bool FFocusDataListItem::ChangeLinkedFocusValues(float NewFocusValue) const
+{
+	if (ULensFile* LensFilePtr = WeakLensFile.Get())
+	{
+		bool bFocusValuesChanged = false;
+		bool bReplaceExisting = false;
+
+		TSharedRef<SWidget> DialogContent = SNew(SCheckBox)
+			.ToolTipText(LOCTEXT("ReplaceExistingZoomPointsInFocusToolTip", "When checked, any existing zoom points in the destination focus will be replaced with those in the source focus"))
+			.IsChecked_Lambda([&bReplaceExisting]() { return bReplaceExisting ? ECheckBoxState::Checked : ECheckBoxState::Unchecked; })
+			.OnCheckStateChanged_Lambda([&bReplaceExisting](ECheckBoxState CheckBoxState) { bReplaceExisting = CheckBoxState == ECheckBoxState::Checked; })
+			.Padding(FMargin(4.0, 0.0))
+			[
+				SNew(STextBlock).Text(LOCTEXT("ReplaceExistingZoomPointsLabel", "Replace existing zoom points?"))
+			];
+		
+		auto OnApplyChange = [LensFilePtr, NewFocusValue, &bReplaceExisting, &bFocusValuesChanged](const TArray<SCameraCalibrationLinkedPointsDialog::FLinkedItem>& LinkedItems)
+		{
+			FScopedTransaction Transaction(LOCTEXT("ChangeLinkedFocusPointsTransaction", "Change Linked Focus Points"));
+			LensFilePtr->Modify();
+
+			for (const SCameraCalibrationLinkedPointsDialog::FLinkedItem& LinkedItem : LinkedItems)
+			{
+				if (LensFilePtr->HasFocusPoint(LinkedItem.Category, NewFocusValue))
+				{
+					LensFilePtr->MergeFocusPoint(LinkedItem.Category, LinkedItem.Focus, NewFocusValue, bReplaceExisting);
+				}
+				else
+				{
+					LensFilePtr->ChangeFocusPoint(LinkedItem.Category, LinkedItem.Focus, NewFocusValue);
+				}
+			}
+
+			bFocusValuesChanged = true;
+		};
+
+		const SCameraCalibrationLinkedPointsDialog::FLinkedItem Item(Category, Focus);
+		TSharedRef<SCameraCalibrationLinkedPointsDialog> DialogBox = SNew(SCameraCalibrationLinkedPointsDialog, LensFilePtr, Item)
+			.LinkedItemMode(SCameraCalibrationLinkedPointsDialog::ELinkedItemMode::Focus)
+			.DialogText(LOCTEXT("ChangeLinkedFocusDialogText", "The calibration data you wish to change may be inherently linked to additional data.\nChoose any and all linked data you wish to change."))
+			.AcceptButtonText(LOCTEXT("ChangeLinkedFocusAcceptButton", "Change Focus"))
+			.OnApplyLinkedAction_Lambda(OnApplyChange)
+			[
+				DialogContent
+			];
+		
+		SCameraCalibrationLinkedPointsDialog::OpenWindow(LOCTEXT("ChangeFocusValueWindowLabel", "Change Focus Value"), DialogBox);
+
+		return bFocusValuesChanged;
+	}
+
+	return false;
+}
+
+bool FFocusDataListItem::RemoveLinkedFocusValues() const
+{
+	if (ULensFile* LensFilePtr = WeakLensFile.Get())
+	{
+		bool bFocusValuesRemoved = false;
+
+		auto OnApplyChange = [LensFilePtr, &bFocusValuesRemoved](const TArray<SCameraCalibrationLinkedPointsDialog::FLinkedItem>& LinkedItems)
+		{
+			FScopedTransaction Transaction(LOCTEXT("RemoveLinkedFocusPointsTransaction", "Remove Points"));
+			LensFilePtr->Modify();
+
+			for (const SCameraCalibrationLinkedPointsDialog::FLinkedItem& LinkedItem : LinkedItems)
+			{
+				if (LinkedItem.Zoom.IsSet())
+				{
+					LensFilePtr->RemoveZoomPoint(LinkedItem.Category, LinkedItem.Focus, LinkedItem.Zoom.GetValue());
+				}
+				else
+				{
+					LensFilePtr->RemoveFocusPoint(LinkedItem.Category, LinkedItem.Focus);
+				}
+			}
+
+			bFocusValuesRemoved = true;
+		};
+
+		const SCameraCalibrationLinkedPointsDialog::FLinkedItem Item(Category, Focus);
+		TSharedRef<SCameraCalibrationLinkedPointsDialog> DialogBox = SNew(SCameraCalibrationLinkedPointsDialog, LensFilePtr, Item)
+			.LinkedItemMode(SCameraCalibrationLinkedPointsDialog::ELinkedItemMode::Both)
+			.DialogText(LOCTEXT("RemoveLinkedFocusDialogText", "The calibration data you wish to delete may be inherently linked to additional data.\nChoose any and all linked data you wish to delete."))
+			.AcceptButtonText(LOCTEXT("RemoveLinkedFocusAcceptButton", "Remove Selected"))
+			.OnApplyLinkedAction_Lambda(OnApplyChange);
+
+		SCameraCalibrationLinkedPointsDialog::OpenWindow(LOCTEXT("RemoveFocusWindowLabel", "Remove Points"), DialogBox);
+		
+		return bFocusValuesRemoved;
+	}
+	
+	return false;
+}
+
+FZoomDataListItem::FZoomDataListItem(ULensFile* InLensFile, ELensDataCategory InCategory, int32 InSubCategoryIndex, const TSharedRef<FFocusDataListItem> InParent, float InZoom, FOnDataChanged InOnDataChangedCallback)
+	: FLensDataListItem(InLensFile, InCategory, InSubCategoryIndex, InOnDataChangedCallback)
 	, Zoom(InZoom)
 	, WeakParent(InParent)
 {
@@ -108,6 +262,8 @@ TSharedRef<ITableRow> FZoomDataListItem::MakeTreeRowWidget(const TSharedRef<STab
 	return SNew(SLensDataItem, InOwnerTable, SharedThis(this))
 		.EntryLabel(LOCTEXT("ZoomLabel", "Zoom: "))
 		.EntryValue(Zoom)
+		.AllowEditEntryValue(true)
+		.OnEntryValueChanged(this, &FZoomDataListItem::OnZoomValueChanged)
 		.AllowRemoval(SubCategoryIndex == INDEX_NONE )
 		.EditPointVisibility(EVisibility::Visible)
 		.AllowEditPoint(MakeAttributeLambda([this]() { return SubCategoryIndex == INDEX_NONE; }));
@@ -124,24 +280,20 @@ void FZoomDataListItem::OnRemoveRequested() const
 			{
 				return;	
 			}
+			
 			if (LinkDataTable->HasLinkedZoomValues(ParentItem->Focus, Zoom))
 			{
-				SCameraCalibrationRemovePointDialog::OpenWindow(
-					LensFilePtr,
-					Category,
-					FSimpleDelegate::CreateLambda([this, ParentItem]()
-					{
-						OnDataRemovedCallback.ExecuteIfBound(ParentItem->Focus, Zoom);
-					}),
-					ParentItem->Focus,
-					Zoom);
+				if (RemoveLinkedZoomValues())
+				{
+					OnDataChangedCallback.ExecuteIfBound(ELensDataChangedReason::DataRemoved, ParentItem->Focus, Zoom);
+				}
 			}
 			else
 			{
 				FScopedTransaction Transaction(LOCTEXT("RemoveZoomPointTransaction", "Remove Zoom Point"));
 				LensFilePtr->Modify();
 				LensFilePtr->RemoveZoomPoint(Category, ParentItem->Focus, Zoom);
-				OnDataRemovedCallback.ExecuteIfBound(ParentItem->Focus, Zoom);
+				OnDataChangedCallback.ExecuteIfBound(ELensDataChangedReason::DataRemoved, ParentItem->Focus, Zoom);
 			}
 		}
 	}
@@ -169,42 +321,201 @@ void FZoomDataListItem::EditItem()
 	{
 		return;
 	}
+
+	FSimpleDelegate OnPointSaved = FSimpleDelegate::CreateLambda([this]()
+	{
+		OnDataChangedCallback.ExecuteIfBound(ELensDataChangedReason::DataChanged, GetFocus().GetValue(), Zoom);	
+	});
 	
 	switch (Category)
 	{
 		case ELensDataCategory::Zoom:
 		{
-			LensDataEditPointDialog::OpenDialog<FFocalLengthInfo>(LensFilePtr, GetFocus().GetValue(), Zoom, LensFilePtr->FocalLengthTable);
+			LensDataEditPointDialog::OpenDialog<FFocalLengthInfo>(LensFilePtr, Category, GetFocus().GetValue(), Zoom, LensFilePtr->FocalLengthTable, OnPointSaved);
 			break;
 		}
 		case ELensDataCategory::ImageCenter:
 		{
-			LensDataEditPointDialog::OpenDialog<FImageCenterInfo>(LensFilePtr, GetFocus().GetValue(), Zoom, LensFilePtr->ImageCenterTable);
+			LensDataEditPointDialog::OpenDialog<FImageCenterInfo>(LensFilePtr, Category, GetFocus().GetValue(), Zoom, LensFilePtr->ImageCenterTable, OnPointSaved);
 			break;
 		}
 		case ELensDataCategory::Distortion:
 		{
-			LensDataEditPointDialog::OpenDialog<FDistortionInfo>(LensFilePtr, GetFocus().GetValue(), Zoom, LensFilePtr->DistortionTable);
+			LensDataEditPointDialog::OpenDialog<FDistortionInfo>(LensFilePtr, Category, GetFocus().GetValue(), Zoom, LensFilePtr->DistortionTable, OnPointSaved);
 			break;
 		}
 		case ELensDataCategory::NodalOffset:
 		{
-			LensDataEditPointDialog::OpenDialog<FNodalPointOffset>(LensFilePtr, GetFocus().GetValue(), Zoom, LensFilePtr->NodalOffsetTable);
+			LensDataEditPointDialog::OpenDialog<FNodalPointOffset>(LensFilePtr, Category, GetFocus().GetValue(), Zoom, LensFilePtr->NodalOffsetTable, OnPointSaved);
 			break;
 		}
 		case ELensDataCategory::STMap:
 		{
-			LensDataEditPointDialog::OpenDialog<FSTMapInfo>(LensFilePtr, GetFocus().GetValue(), Zoom, LensFilePtr->STMapTable);
+			LensDataEditPointDialog::OpenDialog<FSTMapInfo>(LensFilePtr, Category, GetFocus().GetValue(), Zoom, LensFilePtr->STMapTable, OnPointSaved);
 			break;
 		}
 	}
 }
 
+bool FZoomDataListItem::OnZoomValueChanged(float NewZoomValue)
+{
+	if (ULensFile* LensFilePtr = WeakLensFile.Get())
+	{
+		if(TSharedPtr<FFocusDataListItem> ParentItem = WeakParent.Pin())
+		{
+			const FBaseLensTable* const DataTable = LensFilePtr->GetDataTable(Category);
+			if (!ensure(DataTable))
+			{
+				return false;	
+			}
+
+			if (DataTable->HasLinkedZoomValues(ParentItem->Focus, Zoom))
+			{
+				// If there are any linked zoom values, allow the linked zoom dialog box to handle any needed changes to the focuses in the lens file
+				if (ChangeLinkedZoomValues(NewZoomValue))
+				{
+					Zoom = NewZoomValue;
+					OnDataChangedCallback.ExecuteIfBound(ELensDataChangedReason::DataChanged, ParentItem->Focus, Zoom);
+					return true;
+				}
+				
+				return false;
+			}
+			
+			if (LensFilePtr->HasZoomPoint(Category, ParentItem->Focus, NewZoomValue))
+			{
+				if (FCameraCalibrationWidgetHelpers::ShowReplaceZoomWarning())
+				{
+					FScopedTransaction Transaction(LOCTEXT("ReplaceZoomPointTransaction", "Replace Zoom Point"));
+				
+					LensFilePtr->Modify();
+					LensFilePtr->ChangeZoomPoint(Category, ParentItem->Focus, Zoom, NewZoomValue);
+
+					Zoom = NewZoomValue;
+					OnDataChangedCallback.ExecuteIfBound(ELensDataChangedReason::DataChanged, ParentItem->Focus, Zoom);
+
+					return true;
+				}
+			}
+			else
+			{
+				FScopedTransaction Transaction(LOCTEXT("ChangeZoomPointTransaction", "Change Zoom Point"));
+				
+				LensFilePtr->Modify();
+				LensFilePtr->ChangeZoomPoint(Category, ParentItem->Focus, Zoom, NewZoomValue);
+
+				Zoom = NewZoomValue;
+				OnDataChangedCallback.ExecuteIfBound(ELensDataChangedReason::DataChanged, ParentItem->Focus, Zoom);
+
+				return true;
+			}
+		}
+	}
+
+	return false;
+}
+
+bool FZoomDataListItem::ChangeLinkedZoomValues(float NewZoomValue) const
+{
+	if (ULensFile* LensFilePtr = WeakLensFile.Get())
+	{
+		if(TSharedPtr<FFocusDataListItem> ParentItem = WeakParent.Pin())
+		{
+			const float Focus = ParentItem->Focus;
+			bool bZoomValuesChanged = false;
+			bool bReplaceExisting = false;
+
+			TSharedRef<SWidget> DialogContent = SNew(SCheckBox)
+				.ToolTipText(LOCTEXT("ReplaceExistingZoomPointsToolTip", "When checked, any existing zoom points will be replaced with the zoom point being changed"))
+				.IsChecked_Lambda([&bReplaceExisting]() { return bReplaceExisting ? ECheckBoxState::Checked : ECheckBoxState::Unchecked; })
+				.OnCheckStateChanged_Lambda([&bReplaceExisting](ECheckBoxState CheckBoxState) { bReplaceExisting = CheckBoxState == ECheckBoxState::Checked; })
+				.Padding(FMargin(4.0, 0.0))
+				[
+					SNew(STextBlock).Text(LOCTEXT("ReplaceExistingZoomPointsLabel", "Replace existing zoom points?"))
+				];
+			
+			auto OnApplyChange = [LensFilePtr, Focus, NewZoomValue, &bReplaceExisting, &bZoomValuesChanged](const TArray<SCameraCalibrationLinkedPointsDialog::FLinkedItem>& LinkedItems)
+			{
+				FScopedTransaction Transaction(LOCTEXT("ChangeLinkedZoomPointsTransaction", "Change Linked Zoom Points"));
+				LensFilePtr->Modify();
+
+				for (const SCameraCalibrationLinkedPointsDialog::FLinkedItem& LinkedItem : LinkedItems)
+				{
+					if (!LensFilePtr->HasZoomPoint(LinkedItem.Category, Focus, NewZoomValue) || bReplaceExisting)
+					{
+						LensFilePtr->ChangeZoomPoint(LinkedItem.Category, LinkedItem.Focus, LinkedItem.Zoom.GetValue(), NewZoomValue);
+					}
+				}
+				
+				bZoomValuesChanged = true;
+			};
+			
+			const SCameraCalibrationLinkedPointsDialog::FLinkedItem Item(Category, ParentItem->Focus, Zoom);
+			TSharedRef<SCameraCalibrationLinkedPointsDialog> DialogBox = SNew(SCameraCalibrationLinkedPointsDialog, LensFilePtr, Item)
+				.LinkedItemMode(SCameraCalibrationLinkedPointsDialog::ELinkedItemMode::Zoom)
+				.DialogText(LOCTEXT("ChangeLinkedZoomDialogText", "The calibration data you wish to change may be inherently linked to additional data.\nChoose any and all linked data you wish to change."))
+				.AcceptButtonText(LOCTEXT("ChangeLinkedZoomAcceptButton", "Change Zoom"))
+				.OnApplyLinkedAction_Lambda(OnApplyChange)
+				[
+					DialogContent
+				];
+			
+			SCameraCalibrationLinkedPointsDialog::OpenWindow(LOCTEXT("ChangeLinkedZoomValuesWindowLabel", "Change Zoom Value"), DialogBox);
+
+			return bZoomValuesChanged;
+		}
+	}
+
+	return false;
+}
+
+bool FZoomDataListItem::RemoveLinkedZoomValues() const
+{
+	if (ULensFile* LensFilePtr = WeakLensFile.Get())
+	{
+		if (TSharedPtr<FFocusDataListItem> ParentItem = WeakParent.Pin())
+		{
+			bool bZoomValuesRemoved = false;
+
+			auto OnApplyChange = [LensFilePtr, &bZoomValuesRemoved](const TArray<SCameraCalibrationLinkedPointsDialog::FLinkedItem>& LinkedItems)
+			{
+				FScopedTransaction Transaction(LOCTEXT("RemoveLinkedZoomPointsTransaction", "Remove Points"));
+				LensFilePtr->Modify();
+
+				for (const SCameraCalibrationLinkedPointsDialog::FLinkedItem& LinkedItem : LinkedItems)
+				{
+					if (LinkedItem.Zoom.IsSet())
+					{
+						LensFilePtr->RemoveZoomPoint(LinkedItem.Category, LinkedItem.Focus, LinkedItem.Zoom.GetValue());
+					}
+				}
+
+				bZoomValuesRemoved = true;
+			};
+
+			const SCameraCalibrationLinkedPointsDialog::FLinkedItem Item(Category, ParentItem->Focus, Zoom);
+			TSharedRef<SCameraCalibrationLinkedPointsDialog> DialogBox = SNew(SCameraCalibrationLinkedPointsDialog, LensFilePtr, Item)
+				.LinkedItemMode(SCameraCalibrationLinkedPointsDialog::ELinkedItemMode::Zoom)
+				.DialogText(LOCTEXT("RemoveLinkedZoomDialogText", "The calibration data you wish to delete may be inherently linked to additional data.\nChoose any and all linked data you wish to delete."))
+				.AcceptButtonText(LOCTEXT("RemoveLinkedZoomAcceptButton", "Remove Selected"))
+				.OnApplyLinkedAction_Lambda(OnApplyChange);
+
+			SCameraCalibrationLinkedPointsDialog::OpenWindow(LOCTEXT("RemoveZoomWindowLabel", "Remove Points"), DialogBox);
+		
+			return bZoomValuesRemoved;
+		}
+	}
+	
+	return false;
+}
+
 void SLensDataItem::Construct(const FArguments& InArgs, const TSharedRef<STableViewBase>& OwnerTable, const TSharedRef<FLensDataListItem> InItemData)
 {
 	WeakItem = InItemData;
+	OnEntryValueChanged = InArgs._OnEntryValueChanged;
+	EntryValue = InArgs._EntryValue;
 
-	const float EntryValue = InArgs._EntryValue;
+	const TAttribute<bool> AllowEditEntryValue = InArgs._AllowEditEntryValue;
 
 	STableRow<TSharedPtr<FLensDataListItem>>::Construct(
 		STableRow<TSharedPtr<FLensDataListItem>>::FArguments()
@@ -223,8 +534,21 @@ void SLensDataItem::Construct(const FArguments& InArgs, const TSharedRef<STableV
 			.FillWidth(1.0f)
 			.HAlign(HAlign_Left)
 			[
-				SNew(STextBlock)
-				.Text_Lambda([EntryValue](){ return FText::AsNumber(EntryValue); })
+				SNew(SWidgetSwitcher)
+				.WidgetIndex_Lambda([AllowEditEntryValue]() { return AllowEditEntryValue.Get(false) ? 1 : 0; })
+
+				+SWidgetSwitcher::Slot()
+				[
+					SNew(STextBlock)
+					.Text(FText::AsNumber(EntryValue))
+				]
+
+				+SWidgetSwitcher::Slot()
+				[
+					SNew(SNumericEntryBox<float>)
+					.Value_Lambda([this]() { return EntryValue; })
+					.OnValueCommitted(this, &SLensDataItem::OnEntryValueCommitted)
+				]
 			]
 			+ SHorizontalBox::Slot()
 			.AutoWidth()
@@ -266,6 +590,26 @@ FReply SLensDataItem::OnRemovePointClicked() const
 	return FReply::Handled();
 }
 
+void SLensDataItem::OnEntryValueCommitted(float NewValue, ETextCommit::Type CommitType)
+{
+	// Avoid duplicate handling if we are already waiting for a commit to resolve, which could happen
+	// if OnEntryValueChange invoked a modal and is waiting for a response from the user
+	if (bIsCommittingValue)
+	{
+		return;
+	}
+
+	bIsCommittingValue = true;
+	if (EntryValue != NewValue)
+	{
+		if (OnEntryValueChanged.IsBound() && OnEntryValueChanged.Execute(NewValue))
+		{
+			EntryValue = NewValue;
+		}
+	}
+	bIsCommittingValue = false;
+}
+
 FReply SLensDataItem::OnEditPointClicked() const
 {
 	if (TSharedPtr<FLensDataListItem> Item = WeakItem.Pin())
diff --git a/Engine/Plugins/VirtualProduction/CameraCalibration/Source/CameraCalibrationEditor/Private/AssetEditor/SLensDataListItem.h b/Engine/Plugins/VirtualProduction/CameraCalibration/Source/CameraCalibrationEditor/Private/AssetEditor/SLensDataListItem.h
--- a/Engine/Plugins/VirtualProduction/CameraCalibration/Source/CameraCalibrationEditor/Private/AssetEditor/SLensDataListItem.h
+++ b/Engine/Plugins/VirtualProduction/CameraCalibration/Source/CameraCalibrationEditor/Private/AssetEditor/SLensDataListItem.h
@@ -9,7 +9,13 @@
 #include "LensFile.h"
 #include "Widgets/DeclarativeSyntaxSupport.h"
 
-DECLARE_DELEGATE_TwoParams(FOnDataRemoved, float /** Focus */, TOptional<float> /** Possible Zoom */);
+enum class ELensDataChangedReason : uint8
+{
+	DataRemoved,
+	DataChanged
+};
+
+DECLARE_DELEGATE_ThreeParams(FOnDataChanged, ELensDataChangedReason /** ChangedReason */, float /** Focus */, TOptional<float> /** Possible Zoom */);
 
 /**
 * Data entry item
@@ -17,7 +23,7 @@ DECLARE_DELEGATE_TwoParams(FOnDataRemoved, float /** Focus */, TOptional<float>
 class FLensDataListItem : public TSharedFromThis<FLensDataListItem>
 {
 public:
-	FLensDataListItem(ULensFile* InLensFile, ELensDataCategory InCategory, int32 InSubCategoryIndex, FOnDataRemoved InOnDataRemovedCallback);
+	FLensDataListItem(ULensFile* InLensFile, ELensDataCategory InCategory, int32 InSubCategoryIndex, FOnDataChanged InOnDataChangedCallback);
 
 	virtual ~FLensDataListItem() = default;
 	
@@ -39,8 +45,8 @@ public:
 	/** Children of this item */
 	TArray<TSharedPtr<FLensDataListItem>> Children;
 
-	/** Delegate to call when data is removed */
-	FOnDataRemoved OnDataRemovedCallback;
+	/** Delegate to call when data is changed */
+	FOnDataChanged OnDataChangedCallback;
 };
 
 /**
@@ -49,7 +55,7 @@ public:
 class FEncoderDataListItem : public FLensDataListItem
 {
 public:
-	FEncoderDataListItem(ULensFile* InLensFile, ELensDataCategory InCategory, float InInput, int32 InIndex, FOnDataRemoved InOnDataRemovedCallback);
+	FEncoderDataListItem(ULensFile* InLensFile, ELensDataCategory InCategory, float InInput, int32 InIndex, FOnDataChanged InOnDataChangedCallback);
 
 	virtual void OnRemoveRequested() const override;
 	virtual TSharedRef<ITableRow> MakeTreeRowWidget(const TSharedRef<STableViewBase>& InOwnerTable) override;
@@ -67,12 +73,23 @@ public:
 class FFocusDataListItem : public FLensDataListItem
 {
 public:
-	FFocusDataListItem(ULensFile* InLensFile, ELensDataCategory InCategory, int32 InSubCategoryIndex, float InFocus, FOnDataRemoved InOnDataRemovedCallback);
+	FFocusDataListItem(ULensFile* InLensFile, ELensDataCategory InCategory, int32 InSubCategoryIndex, float InFocus, FOnDataChanged InOnDataChangedCallback);
 
 	virtual TSharedRef<ITableRow> MakeTreeRowWidget(const TSharedRef<STableViewBase>& InOwnerTable) override;
 	virtual void OnRemoveRequested() const override;
 	virtual TOptional<float> GetFocus() const override { return Focus; }
 
+	/** Raised when the focus value is changed on this focus item */
+	bool OnFocusValueChanged(float NewFocusValue);
+
+	/** Creates a dialog box that allows users to change linked focus values when this focus item is changed
+	 * @returns true if the user presses the accept button on the dialog, false otherwise */
+	bool ChangeLinkedFocusValues(float NewFocusValue) const;
+
+	/** Creates a dialog box that allows users to remove linked focus and zoom values when this focus item is removed
+	 * @returns true if the user presses the accept button on the dialog, false otherwise */
+	bool RemoveLinkedFocusValues() const;
+	
 	/** Focus value of this item */
 	float Focus;
 };
@@ -83,7 +100,7 @@ public:
 class FZoomDataListItem : public FLensDataListItem
 {
 public:
-	FZoomDataListItem(ULensFile* InLensFile, ELensDataCategory InCategory, int32 InSubCategoryIndex, const TSharedRef<FFocusDataListItem> InParent, float InZoom, FOnDataRemoved InOnDataRemovedCallback);
+	FZoomDataListItem(ULensFile* InLensFile, ELensDataCategory InCategory, int32 InSubCategoryIndex, const TSharedRef<FFocusDataListItem> InParent, float InZoom, FOnDataChanged InOnDataChangedCallback);
 
 	//~ Begin FLensDataListItem interface
 	virtual TSharedRef<ITableRow> MakeTreeRowWidget(const TSharedRef<STableViewBase>& InOwnerTable) override;
@@ -92,6 +109,17 @@ public:
 	virtual void EditItem() override;
 	//~ End FLensDataListItem interface
 
+	/** Raised when the zoom value is changed on this zoom item */
+	bool OnZoomValueChanged(float NewZoomValue);
+
+	/** Creates a dialog box that allows users to change linked zoom values when this zoom item is changed
+	 * @returns true if the user presses the accept button on the dialog, false otherwise */
+	bool ChangeLinkedZoomValues(float NewZoomValue) const;
+
+	/** Creates a dialog box that allows users to remove linked zoom values when this zoom item is removed
+	 * @returns true if the user presses the accept button on the dialog, false otherwise */
+	bool RemoveLinkedZoomValues() const;
+	
 	/** Zoom value of this item */
 	float Zoom = 0.0f;
 
@@ -104,6 +132,10 @@ public:
  */
 class SLensDataItem : public STableRow<TSharedPtr<FLensDataListItem>>
 {
+public:
+	/** Delegate raised when the entry value of the data item has been changed. Returns whether the change should be committed or not */
+	DECLARE_DELEGATE_RetVal_OneParam(bool, FOnEntryValueChanged, float /*NewValue*/);
+
 	SLATE_BEGIN_ARGS(SLensDataItem)
 		:  _EntryLabel(FText::GetEmpty())
 		,  _EntryValue(0.f)
@@ -123,6 +155,10 @@ class SLensDataItem : public STableRow<TSharedPtr<FLensDataListItem>>
 
 		/** Whether Item point editable */
 		SLATE_ATTRIBUTE(bool, AllowEditPoint)
+
+		/** Whether the entry's value should be editable */
+		SLATE_ATTRIBUTE(bool, AllowEditEntryValue)
+		SLATE_EVENT(FOnEntryValueChanged, OnEntryValueChanged)
 	SLATE_END_ARGS()
 
 	void Construct(const FArguments& InArgs, const TSharedRef<STableViewBase>& OwnerTable, const TSharedRef<FLensDataListItem> InItemData);
@@ -134,8 +170,14 @@ private:
 	/** Remove Button Handler */
 	FReply OnRemovePointClicked() const;
 
+	void OnEntryValueCommitted(float NewValue, ETextCommit::Type CommitType);
+
 private:
 
 	/** WeakPtr to source data item */
 	TWeakPtr<FLensDataListItem> WeakItem;
+
+	float EntryValue = 0.0;
+	bool bIsCommittingValue = false;
+	FOnEntryValueChanged OnEntryValueChanged;
 };
diff --git a/Engine/Plugins/VirtualProduction/CameraCalibration/Source/CameraCalibrationEditor/Private/AssetEditor/SLensDataViewer.cpp b/Engine/Plugins/VirtualProduction/CameraCalibration/Source/CameraCalibrationEditor/Private/AssetEditor/SLensDataViewer.cpp
--- a/Engine/Plugins/VirtualProduction/CameraCalibration/Source/CameraCalibrationEditor/Private/AssetEditor/SLensDataViewer.cpp
+++ b/Engine/Plugins/VirtualProduction/CameraCalibration/Source/CameraCalibrationEditor/Private/AssetEditor/SLensDataViewer.cpp
@@ -54,19 +54,19 @@ namespace LensDataUtils
 	const FName RotationZLabel(TEXT("Roll"));
 
 	template<typename TFocusPoint>
-	void MakeFocusEntries(ULensFile* InLensFile, ELensDataCategory InCategory, int32 InSubCategoryIndex, TConstArrayView<TFocusPoint> FocusPoints, TArray<TSharedPtr<FLensDataListItem>>& OutDataItems, FOnDataRemoved InDataRemovedCallback)
+	void MakeFocusEntries(ULensFile* InLensFile, ELensDataCategory InCategory, int32 InSubCategoryIndex, TConstArrayView<TFocusPoint> FocusPoints, TArray<TSharedPtr<FLensDataListItem>>& OutDataItems, FOnDataChanged InDataChangedCallback)
 	{
 		OutDataItems.Reserve(FocusPoints.Num());
 		for (const TFocusPoint& Point : FocusPoints)
 		{
 			//Add entry for focus
-			TSharedPtr<FFocusDataListItem> CurrentFocus = MakeShared<FFocusDataListItem>(InLensFile, InCategory, InSubCategoryIndex, Point.Focus, InDataRemovedCallback);
+			TSharedPtr<FFocusDataListItem> CurrentFocus = MakeShared<FFocusDataListItem>(InLensFile, InCategory, InSubCategoryIndex, Point.Focus, InDataChangedCallback);
 			OutDataItems.Add(CurrentFocus);
 
 			for(int32 Index = 0; Index < Point.GetNumPoints(); ++Index)
 			{
 				//Add zoom points for this focus
-				TSharedPtr<FZoomDataListItem> ZoomItem = MakeShared<FZoomDataListItem>(InLensFile, InCategory, InSubCategoryIndex, CurrentFocus.ToSharedRef(),Point.GetZoom(Index), InDataRemovedCallback);
+				TSharedPtr<FZoomDataListItem> ZoomItem = MakeShared<FZoomDataListItem>(InLensFile, InCategory, InSubCategoryIndex, CurrentFocus.ToSharedRef(),Point.GetZoom(Index), InDataChangedCallback);
 				CurrentFocus->Children.Add(ZoomItem);
 			}
 		}
@@ -485,11 +485,25 @@ void SLensDataViewer::RefreshDataEntriesTree()
 {
 	TSharedPtr<FLensDataListItem> CurrentSelection = GetSelectedDataEntry();
 
+	// Save the items that are expanded, so that the expanded state can be restored after the tree has been refreshed
+	TSet<TSharedPtr<FLensDataListItem>> ExpandedItems;
+	DataEntriesTree->GetExpandedItems(ExpandedItems);
+
+	TArray<float> FocusesToExpand;
+	for (const TSharedPtr<FLensDataListItem>& Item : ExpandedItems)
+	{
+		TOptional<float> ItemFocus = Item->GetFocus();
+		if (ItemFocus.IsSet())
+		{
+			FocusesToExpand.Add(ItemFocus.GetValue());
+		}
+	}
+	
 	DataEntries.Reset();
 
 	if (TSharedPtr<FLensDataCategoryItem> CategoryItem = GetDataCategorySelection())
 	{
-		FOnDataRemoved DataRemovedCallback = FOnDataRemoved::CreateSP(this, &SLensDataViewer::OnDataPointRemoved);
+		FOnDataChanged DataChangedCallback = FOnDataChanged::CreateSP(this, &SLensDataViewer::OnDataPointChanged);
 
 		switch (CategoryItem->Category)
 		{
@@ -497,7 +511,7 @@ void SLensDataViewer::RefreshDataEntriesTree()
 			{
 				for (int32 Index = 0; Index <LensFile->EncodersTable.GetNumFocusPoints(); ++Index)
 				{
-					DataEntries.Add(MakeShared<FEncoderDataListItem>(LensFile.Get(), CategoryItem->Category, LensFile->EncodersTable.GetFocusInput(Index), Index, DataRemovedCallback));
+					DataEntries.Add(MakeShared<FEncoderDataListItem>(LensFile.Get(), CategoryItem->Category, LensFile->EncodersTable.GetFocusInput(Index), Index, DataChangedCallback));
 				}
 				break;
 			}
@@ -505,38 +519,38 @@ void SLensDataViewer::RefreshDataEntriesTree()
 			{
 				for (int32 Index = 0; Index <LensFile->EncodersTable.GetNumIrisPoints(); ++Index)
 				{
-					DataEntries.Add(MakeShared<FEncoderDataListItem>(LensFile.Get(), CategoryItem->Category, LensFile->EncodersTable.GetIrisInput(Index), Index, DataRemovedCallback));
+					DataEntries.Add(MakeShared<FEncoderDataListItem>(LensFile.Get(), CategoryItem->Category, LensFile->EncodersTable.GetIrisInput(Index), Index, DataChangedCallback));
 				}
 				break;
 			}
 			case ELensDataCategory::Zoom:
 			{
 				const TConstArrayView<FFocalLengthFocusPoint> FocusPoints = LensFile->FocalLengthTable.GetFocusPoints();
-				LensDataUtils::MakeFocusEntries(LensFile.Get(), CategoryItem->Category, CategoryItem->GetParameterIndex(), FocusPoints, DataEntries, DataRemovedCallback);
+				LensDataUtils::MakeFocusEntries(LensFile.Get(), CategoryItem->Category, CategoryItem->GetParameterIndex(), FocusPoints, DataEntries, DataChangedCallback);
 				break;
 			}
 			case ELensDataCategory::Distortion:
 			{
 				const TConstArrayView<FDistortionFocusPoint> FocusPoints = LensFile->DistortionTable.GetFocusPoints();
-				LensDataUtils::MakeFocusEntries(LensFile.Get(), CategoryItem->Category, CategoryItem->GetParameterIndex(), FocusPoints, DataEntries, DataRemovedCallback);
+				LensDataUtils::MakeFocusEntries(LensFile.Get(), CategoryItem->Category, CategoryItem->GetParameterIndex(), FocusPoints, DataEntries, DataChangedCallback);
 				break;
 			}
 			case ELensDataCategory::ImageCenter:
 			{
 				const TConstArrayView<FImageCenterFocusPoint> FocusPoints = LensFile->ImageCenterTable.GetFocusPoints();
-				LensDataUtils::MakeFocusEntries(LensFile.Get(), CategoryItem->Category, CategoryItem->GetParameterIndex(), FocusPoints, DataEntries, DataRemovedCallback);
+				LensDataUtils::MakeFocusEntries(LensFile.Get(), CategoryItem->Category, CategoryItem->GetParameterIndex(), FocusPoints, DataEntries, DataChangedCallback);
 				break;
 			}
 			case ELensDataCategory::NodalOffset:
 			{
 				const TConstArrayView<FNodalOffsetFocusPoint> Points = LensFile->NodalOffsetTable.GetFocusPoints();
-				LensDataUtils::MakeFocusEntries(LensFile.Get(), CategoryItem->Category, CategoryItem->GetParameterIndex(), Points, DataEntries, DataRemovedCallback);
+				LensDataUtils::MakeFocusEntries(LensFile.Get(), CategoryItem->Category, CategoryItem->GetParameterIndex(), Points, DataEntries, DataChangedCallback);
 				break;
 			}
 			case ELensDataCategory::STMap:
 			{
 				const TConstArrayView<FSTMapFocusPoint> Points = LensFile->STMapTable.GetFocusPoints();
-				LensDataUtils::MakeFocusEntries(LensFile.Get(), CategoryItem->Category, CategoryItem->GetParameterIndex(), Points, DataEntries, DataRemovedCallback);
+				LensDataUtils::MakeFocusEntries(LensFile.Get(), CategoryItem->Category, CategoryItem->GetParameterIndex(), Points, DataEntries, DataChangedCallback);
 				break;
 			}
 		}
@@ -545,6 +559,16 @@ void SLensDataViewer::RefreshDataEntriesTree()
 	//When data entries have been repopulated, refresh the tree and select first item
 	DataEntriesTree->RequestListRefresh();
 
+	// Restore the expanded items by focus. 
+	for (const TSharedPtr<FLensDataListItem>& Item : DataEntries)
+	{
+		TOptional<float> ItemFocus = Item->GetFocus();
+		if (ItemFocus.IsSet() && FocusesToExpand.Contains(ItemFocus.GetValue()))
+		{
+			DataEntriesTree->SetItemExpansion(Item, true);
+		}
+	}
+	
 	//Try to put back the same selected Focus/Zoom item
 	UpdateDataSelection(CurrentSelection);
 }
@@ -687,7 +711,7 @@ void SLensDataViewer::OnLensDataPointAdded()
 	RefreshDataEntriesTree();
 }
 
-void SLensDataViewer::OnDataPointRemoved(float InFocus, TOptional<float> InZoom)
+void SLensDataViewer::OnDataPointChanged(ELensDataChangedReason ChangedReason, float InFocus, TOptional<float> InZoom)
 {
 	RefreshDataEntriesTree();
 }
diff --git a/Engine/Plugins/VirtualProduction/CameraCalibration/Source/CameraCalibrationEditor/Private/AssetEditor/SLensDataViewer.h b/Engine/Plugins/VirtualProduction/CameraCalibration/Source/CameraCalibrationEditor/Private/AssetEditor/SLensDataViewer.h
--- a/Engine/Plugins/VirtualProduction/CameraCalibration/Source/CameraCalibrationEditor/Private/AssetEditor/SLensDataViewer.h
+++ b/Engine/Plugins/VirtualProduction/CameraCalibration/Source/CameraCalibrationEditor/Private/AssetEditor/SLensDataViewer.h
@@ -18,6 +18,7 @@ class FCameraCalibrationCurveEditor;
 class SCameraCalibrationCurveEditorPanel;
 class FCameraCalibrationStepsController;
 class FCameraCalibrationTimeSliderController;
+enum class ELensDataChangedReason : uint8;
 
 /** Widget used to display data from the LensFile */
 class SLensDataViewer : public SCompoundWidget, public FSelfRegisteringEditorUndoClient
@@ -97,8 +98,8 @@ private:
 	/** Callbacked when user clicks AddPoint from the dialog */
 	void OnLensDataPointAdded();
 
-	/** Callbacked when user clicks remove point buttons */
-	void OnDataPointRemoved(float InFocus, TOptional<float> InZoom);
+	/** Callbacked when user changes anything about the data point that may require a refresh */
+	void OnDataPointChanged(ELensDataChangedReason ChangedReason, float InFocus, TOptional<float> InZoom);
 	
 	/** Called when the data table points list of a data category was updated */
 	void OnDataTablePointsUpdated(ELensDataCategory InCategory);
@@ -119,7 +120,7 @@ private:
 	 * If category hasn't changed, don't recreate data entries
 	 */
 	TSharedPtr<FLensDataCategoryItem> CachedSelectedCategoryItem;
-
+	
 	/** Data items associated with selected data category TreeView */
 	TSharedPtr<STreeView<TSharedPtr<FLensDataListItem>>> DataEntriesTree;
 
diff --git a/Engine/Plugins/VirtualProduction/CameraCalibration/Source/CameraCalibrationEditor/Private/UI/CameraCalibrationWidgetHelpers.cpp b/Engine/Plugins/VirtualProduction/CameraCalibration/Source/CameraCalibrationEditor/Private/UI/CameraCalibrationWidgetHelpers.cpp
--- a/Engine/Plugins/VirtualProduction/CameraCalibration/Source/CameraCalibrationEditor/Private/UI/CameraCalibrationWidgetHelpers.cpp
+++ b/Engine/Plugins/VirtualProduction/CameraCalibration/Source/CameraCalibrationEditor/Private/UI/CameraCalibrationWidgetHelpers.cpp
@@ -9,6 +9,7 @@
 #include "Widgets/Layout/SBox.h"
 #include "Widgets/SBoxPanel.h"
 #include "Widgets/SWidget.h"
+#include "Widgets/Input/SCheckBox.h"
 #include "Widgets/Text/STextBlock.h"
 
 #define LOCTEXT_NAMESPACE "CameraCalibrationWidgetHelpers"
@@ -111,4 +112,63 @@ void FCameraCalibrationWidgetHelpers::DisplayTextureInWindowAlmostFullScreen(UTe
 	}
 }
 
+bool FCameraCalibrationWidgetHelpers::ShowMergeFocusWarning(bool& bOutReplaceExistingZoomPoints)
+{
+	bool bReplaceExistingZoomPoints = false;
+	TSharedRef<SCustomDialog> Dialog = SNew(SCustomDialog)
+		.Title(FText(LOCTEXT("FocusMergeWarningTitle", "Merge existing focus point?")))
+		.ContentAreaPadding(16.0)
+		.Content()
+		[
+			SNew(SVerticalBox)
+					
+			+SVerticalBox::Slot()
+			[
+				SNew(STextBlock)
+				.Text(LOCTEXT("FocusMergeWarningLabel", "A focus point already exists with that value. Would you like to merge this point with that point?"))
+			]
+
+			+SVerticalBox::Slot()
+			.Padding(0.0, 4.0, 0.0, 0.0)
+			[
+				SNew(SCheckBox)
+				.ToolTipText(LOCTEXT("ReplaceExistingToolTip", "When checked, any existing zoom points in the destination focus will be replaced with those in the source focus"))
+				.IsChecked_Lambda([&bReplaceExistingZoomPoints]() { return bReplaceExistingZoomPoints ? ECheckBoxState::Checked : ECheckBoxState::Unchecked; })
+				.OnCheckStateChanged_Lambda([&bReplaceExistingZoomPoints](ECheckBoxState CheckBoxState) { bReplaceExistingZoomPoints = CheckBoxState == ECheckBoxState::Checked; })
+				.Padding(FMargin(4.0, 0.0))
+				[
+					SNew(STextBlock).Text(LOCTEXT("ReplaceExistingLabel", "Replace existing zoom points?"))
+				]
+			]
+		]
+		.Buttons({
+			SCustomDialog::FButton(LOCTEXT("MergeButtonLabel", "Merge")),
+			SCustomDialog::FButton(LOCTEXT("CancelButtonLabel", "Cancel")) }
+		);
+
+	// Dialog result corresponds to the button that was pressed (e.g. 0 means first button, 1 second, -1 no buttons were pressed)
+	const int32 DialogResult = Dialog->ShowModal();
+	bOutReplaceExistingZoomPoints = bReplaceExistingZoomPoints;
+	return DialogResult == 0;
+}
+
+bool FCameraCalibrationWidgetHelpers::ShowReplaceZoomWarning()
+{
+	TSharedRef<SCustomDialog> Dialog = SNew(SCustomDialog)
+		.Title(FText(LOCTEXT("ZoomReplaceWarningTitle", "Replace existing zoom point?")))
+		.Content()
+		[
+			SNew(STextBlock)
+			.Text(LOCTEXT("ZoomReplaceWarningLabel", "A point with that zoom value already exists, would you like to replace it with this point?"))
+		]
+		.Buttons({
+			SCustomDialog::FButton(LOCTEXT("ReplaceButtonLabel", "Replace")),
+			SCustomDialog::FButton(LOCTEXT("CancelButtonLabel", "Cancel")) }
+		);
+
+	// Dialog result corresponds to the button that was pressed (e.g. 0 means first button, 1 second, -1 no buttons were pressed)
+	const int32 DialogResult = Dialog->ShowModal();
+	return DialogResult == 0;
+}
+
 #undef LOCTEXT_NAMESPACE
diff --git a/Engine/Plugins/VirtualProduction/CameraCalibration/Source/CameraCalibrationEditor/Private/UI/CameraCalibrationWidgetHelpers.h b/Engine/Plugins/VirtualProduction/CameraCalibration/Source/CameraCalibrationEditor/Private/UI/CameraCalibrationWidgetHelpers.h
--- a/Engine/Plugins/VirtualProduction/CameraCalibration/Source/CameraCalibrationEditor/Private/UI/CameraCalibrationWidgetHelpers.h
+++ b/Engine/Plugins/VirtualProduction/CameraCalibration/Source/CameraCalibrationEditor/Private/UI/CameraCalibrationWidgetHelpers.h
@@ -24,4 +24,13 @@ public:
 
 	/** Displays a window with the given texture, preserving aspect ratio and almost full screen */
 	static void DisplayTextureInWindowAlmostFullScreen(UTexture2D* Texture, FText&& Title, float ScreenMarginFactor = 0.85f);
+
+	/** Displays a warning dialog box asking if a user wants to merge one focus point into another
+	 * @param bOutReplaceExistingZoomPoints Whether the user wants to replace existing zoom points in the merge
+	 * @return true if the user accepted the merge, false otherwise */
+	static bool ShowMergeFocusWarning(bool& bOutReplaceExistingZoomPoints);
+
+	/** Displays a warning dialog box asking if the user wants to replace one zoom point with another
+	 * @return true if the user accepted the replace, false otherwise */
+	static bool ShowReplaceZoomWarning();
 };
\ No newline at end of file
diff --git a/Engine/Plugins/VirtualProduction/CameraCalibrationCore/Source/CameraCalibrationCore/Private/LensFile.cpp b/Engine/Plugins/VirtualProduction/CameraCalibrationCore/Source/CameraCalibrationCore/Private/LensFile.cpp
--- a/Engine/Plugins/VirtualProduction/CameraCalibrationCore/Source/CameraCalibrationCore/Private/LensFile.cpp
+++ b/Engine/Plugins/VirtualProduction/CameraCalibrationCore/Source/CameraCalibrationCore/Private/LensFile.cpp
@@ -1086,35 +1086,27 @@ void ULensFile::AddSTMapPoint(float NewFocus, float NewZoom, const FSTMapInfo& N
 	STMapTable.AddPoint(NewFocus, NewZoom, NewPoint, InputTolerance, false);
 }
 
+// Helper macros that fill in switch cases for the base lens tables for calling a function
+#define FUNCTION_LENS_TABLE_CASE(DataCategoryName, TableName, FuncCall) case DataCategoryName: TableName.FuncCall; break;
+#define SWITCH_FUNCTION_ON_BASE_LENS_TABLES(DataCategory, FuncCall) FUNCTION_LENS_TABLE_CASE(ELensDataCategory::Distortion, DistortionTable, FuncCall) \
+	FUNCTION_LENS_TABLE_CASE(ELensDataCategory::ImageCenter, ImageCenterTable, FuncCall) \
+	FUNCTION_LENS_TABLE_CASE(ELensDataCategory::Zoom, FocalLengthTable, FuncCall) \
+	FUNCTION_LENS_TABLE_CASE(ELensDataCategory::STMap, STMapTable, FuncCall) \
+	FUNCTION_LENS_TABLE_CASE(ELensDataCategory::NodalOffset, NodalOffsetTable, FuncCall) \
+
+#define FUNCTION_RETVAL_LENS_TABLE_CASE(DataCategoryName, TableName, FuncCall) case DataCategoryName: return TableName.FuncCall;
+#define SWITCH_FUNCTION_RETVAL_ON_BASE_LENS_TABLES(DataCategory, FuncCall) FUNCTION_RETVAL_LENS_TABLE_CASE(ELensDataCategory::Distortion, DistortionTable, FuncCall) \
+	FUNCTION_RETVAL_LENS_TABLE_CASE(ELensDataCategory::ImageCenter, ImageCenterTable, FuncCall) \
+	FUNCTION_RETVAL_LENS_TABLE_CASE(ELensDataCategory::Zoom, FocalLengthTable, FuncCall) \
+	FUNCTION_RETVAL_LENS_TABLE_CASE(ELensDataCategory::STMap, STMapTable, FuncCall) \
+	FUNCTION_RETVAL_LENS_TABLE_CASE(ELensDataCategory::NodalOffset, NodalOffsetTable, FuncCall) \
+
 void ULensFile::RemoveFocusPoint(ELensDataCategory InDataCategory, float InFocus)
 {
 	switch(InDataCategory)
 	{
-		case ELensDataCategory::Distortion:
-		{
-			DistortionTable.RemoveFocusPoint(InFocus);
-			break;
-		}
-		case ELensDataCategory::ImageCenter:
-		{
-			ImageCenterTable.RemoveFocusPoint(InFocus);
-			break;
-		}
-		case ELensDataCategory::Zoom:
-		{
-			FocalLengthTable.RemoveFocusPoint(InFocus);
-			break;
-		}
-		case ELensDataCategory::STMap:
-		{
-			STMapTable.RemoveFocusPoint(InFocus);
-			break;
-		}
-		case ELensDataCategory::NodalOffset:
-		{
-			NodalOffsetTable.RemoveFocusPoint(InFocus);
-			break;
-		}
+		SWITCH_FUNCTION_ON_BASE_LENS_TABLES(InDataCategory, RemoveFocusPoint(InFocus))
+		
 		case ELensDataCategory::Focus:
 		{
 			EncodersTable.RemoveFocusPoint(InFocus);
@@ -1130,35 +1122,87 @@ void ULensFile::RemoveFocusPoint(ELensDataCategory InDataCategory, float InFocus
 	}
 }
 
-void ULensFile::RemoveZoomPoint(ELensDataCategory InDataCategory, float InFocus, float InZoom)
+bool ULensFile::HasFocusPoint(ELensDataCategory InDataCategory, float InFocus) const
 {
 	switch(InDataCategory)
 	{
-		case ELensDataCategory::Distortion:
-		{
-			DistortionTable.RemoveZoomPoint(InFocus, InZoom);
-			break;
-		}
-		case ELensDataCategory::ImageCenter:
-		{
-			ImageCenterTable.RemoveZoomPoint(InFocus, InZoom);
-			break;
-		}
-		case ELensDataCategory::Zoom:
-		{
-			FocalLengthTable.RemoveZoomPoint(InFocus, InZoom);
-			break;
-		}
-		case ELensDataCategory::STMap:
+		SWITCH_FUNCTION_RETVAL_ON_BASE_LENS_TABLES(InDataCategory, HasFocusPoint(InFocus, InputTolerance))
+
+		// Unsupported on encoder tables
+		case ELensDataCategory::Focus:
+		case ELensDataCategory::Iris:
+		default:
 		{
-			STMapTable.RemoveZoomPoint(InFocus, InZoom);
-			break;
+			return false;
 		}
-		case ELensDataCategory::NodalOffset:
+	}
+}
+
+void ULensFile::ChangeFocusPoint(ELensDataCategory InDataCategory, float InExistingFocus, float InNewFocus)
+{
+	switch(InDataCategory)
+	{
+		SWITCH_FUNCTION_ON_BASE_LENS_TABLES(InDataCategory, ChangeFocusPoint(InExistingFocus, InNewFocus, InputTolerance))
+
+		// Changing focus points is unsupported on encoder tables
+		case ELensDataCategory::Focus:
+		case ELensDataCategory::Iris:
+		default:
+		{}
+	}
+}
+
+void ULensFile::MergeFocusPoint(ELensDataCategory InDataCategory, float InSrcFocus, float InDestFocus, bool bReplaceExistingZoomPoints)
+{
+	switch(InDataCategory)
+	{
+		SWITCH_FUNCTION_ON_BASE_LENS_TABLES(InDataCategory, MergeFocusPoint(InSrcFocus, InDestFocus, bReplaceExistingZoomPoints, InputTolerance))
+
+		// Merging focus points is unsupported on encoder tables
+		case ELensDataCategory::Focus:
+		case ELensDataCategory::Iris:
+		default:
+		{}
+	}
+}
+
+void ULensFile::RemoveZoomPoint(ELensDataCategory InDataCategory, float InFocus, float InZoom)
+{
+	switch(InDataCategory)
+	{
+		SWITCH_FUNCTION_ON_BASE_LENS_TABLES(InDataCategory, RemoveZoomPoint(InFocus, InZoom))
+
+		// Encoder tables don't have zoom points
+		case ELensDataCategory::Focus:
+		case ELensDataCategory::Iris:
+		default:
+		{}
+	}
+}
+
+bool ULensFile::HasZoomPoint(ELensDataCategory InDataCategory, float InFocus, float InZoom)
+{
+	switch(InDataCategory)
+	{
+		SWITCH_FUNCTION_RETVAL_ON_BASE_LENS_TABLES(InDataCategory, HasZoomPoint(InFocus, InZoom, InputTolerance))
+
+		// Encoder tables don't have zoom points
+		case ELensDataCategory::Focus:
+		case ELensDataCategory::Iris:
+		default:
 		{
-			NodalOffsetTable.RemoveZoomPoint(InFocus, InZoom);
-			break;
+			return false;
 		}
+	}
+}
+
+void ULensFile::ChangeZoomPoint(ELensDataCategory InDataCategory, float InFocus, float InExistingZoom, float InNewZoom)
+{
+	switch(InDataCategory)
+	{
+		SWITCH_FUNCTION_ON_BASE_LENS_TABLES(InDataCategory, ChangeZoomPoint(InFocus, InExistingZoom, InNewZoom, InputTolerance))
+
+		// Encoder tables don't have zoom points
 		case ELensDataCategory::Focus:
 		case ELensDataCategory::Iris:
 		default:
diff --git a/Engine/Plugins/VirtualProduction/CameraCalibrationCore/Source/CameraCalibrationCore/Private/Tables/DistortionParametersTable.cpp b/Engine/Plugins/VirtualProduction/CameraCalibrationCore/Source/CameraCalibrationCore/Private/Tables/DistortionParametersTable.cpp
--- a/Engine/Plugins/VirtualProduction/CameraCalibrationCore/Source/CameraCalibrationCore/Private/Tables/DistortionParametersTable.cpp
+++ b/Engine/Plugins/VirtualProduction/CameraCalibrationCore/Source/CameraCalibrationCore/Private/Tables/DistortionParametersTable.cpp
@@ -171,11 +171,36 @@ void FDistortionTable::RemoveFocusPoint(float InFocus)
 	LensDataTableUtils::RemoveFocusPoint(FocusPoints, InFocus);
 }
 
+bool FDistortionTable::HasFocusPoint(float InFocus, float InputTolerance) const
+{
+	return DoesFocusPointExists(InFocus, InputTolerance);
+}
+
+void FDistortionTable::ChangeFocusPoint(float InExistingFocus, float InNewFocus, float InputTolerance)
+{
+	LensDataTableUtils::ChangeFocusPoint(FocusPoints, InExistingFocus, InNewFocus, InputTolerance);
+}
+
+void FDistortionTable::MergeFocusPoint(float InSrcFocus, float InDestFocus, bool bReplaceExistingZoomPoints, float InputTolerance)
+{
+	LensDataTableUtils::MergeFocusPoint(FocusPoints, InSrcFocus, InDestFocus, bReplaceExistingZoomPoints, InputTolerance);
+}
+
 void FDistortionTable::RemoveZoomPoint(float InFocus, float InZoom)
 {
 	LensDataTableUtils::RemoveZoomPoint(FocusPoints, InFocus, InZoom);
 }
 
+bool FDistortionTable::HasZoomPoint(float InFocus, float InZoom, float InputTolerance)
+{
+	return DoesZoomPointExists(InFocus, InZoom, InputTolerance);
+}
+
+void FDistortionTable::ChangeZoomPoint(float InFocus, float InExistingZoom, float InNewZoom, float InputTolerance)
+{
+	LensDataTableUtils::ChangeZoomPoint(FocusPoints, InFocus, InExistingZoom, InNewZoom, InputTolerance);
+}
+
 bool FDistortionTable::DoesFocusPointExists(float InFocus, float InputTolerance) const
 {
 	if (GetFocusPoint(InFocus, InputTolerance) != nullptr)
diff --git a/Engine/Plugins/VirtualProduction/CameraCalibrationCore/Source/CameraCalibrationCore/Private/Tables/FocalLengthTable.cpp b/Engine/Plugins/VirtualProduction/CameraCalibrationCore/Source/CameraCalibrationCore/Private/Tables/FocalLengthTable.cpp
--- a/Engine/Plugins/VirtualProduction/CameraCalibrationCore/Source/CameraCalibrationCore/Private/Tables/FocalLengthTable.cpp
+++ b/Engine/Plugins/VirtualProduction/CameraCalibrationCore/Source/CameraCalibrationCore/Private/Tables/FocalLengthTable.cpp
@@ -15,7 +15,7 @@ float FFocalLengthFocusPoint::GetZoom(int32 Index) const
 	return Fx.Keys[Index].Time;
 }
 
-bool FFocalLengthFocusPoint::GetPoint(float InZoom, FFocalLengthZoomPoint& OutZoomPont, float InputTolerance) const
+bool FFocalLengthFocusPoint::GetPoint(float InZoom, FFocalLengthInfo& OutData, float InputTolerance) const
 {
 	const FKeyHandle FxHandle = Fx.FindKey(InZoom, InputTolerance);
 	if(FxHandle != FKeyHandle::Invalid())
@@ -24,9 +24,8 @@ bool FFocalLengthFocusPoint::GetPoint(float InZoom, FFocalLengthZoomPoint& OutZo
 		const int32 PointIndex = Fx.GetIndexSafe(FxHandle);
 		check(FyHandle != FKeyHandle::Invalid() && ZoomPoints.IsValidIndex(PointIndex))
 
-		OutZoomPont.FocalLengthInfo.FxFy.X = Fx.GetKeyValue(FxHandle);
-		OutZoomPont.FocalLengthInfo.FxFy.Y = Fy.GetKeyValue(FyHandle);
-		OutZoomPont.bIsCalibrationPoint = ZoomPoints[PointIndex].bIsCalibrationPoint;
+		OutData.FxFy.X = Fx.GetKeyValue(FxHandle);
+		OutData.FxFy.Y = Fy.GetKeyValue(FyHandle);
 
 		return true;
 	}
@@ -80,6 +79,20 @@ bool FFocalLengthFocusPoint::SetPoint(float InZoom, const FFocalLengthInfo& InDa
 	return false;
 }
 
+bool FFocalLengthFocusPoint::IsCalibrationPoint(float InZoom, float InputTolerance)
+{
+	const FKeyHandle FxHandle = Fx.FindKey(InZoom, InputTolerance);
+	if (FxHandle != FKeyHandle::Invalid())
+	{
+		const int32 PointIndex = Fx.GetIndexSafe(FxHandle);
+		check(ZoomPoints.IsValidIndex(PointIndex))
+
+		return ZoomPoints[PointIndex].bIsCalibrationPoint;
+	}
+
+	return false;
+}
+
 bool FFocalLengthFocusPoint::GetValue(int32 Index, FFocalLengthInfo& OutData) const
 {
 	if(ZoomPoints.IsValidIndex(Index))
@@ -211,11 +224,36 @@ void FFocalLengthTable::RemoveFocusPoint(float InFocus)
 	LensDataTableUtils::RemoveFocusPoint(FocusPoints, InFocus);
 }
 
+bool FFocalLengthTable::HasFocusPoint(float InFocus,  float InputTolerance) const
+{
+	return DoesFocusPointExists(InFocus, InputTolerance);
+}
+
+void FFocalLengthTable::ChangeFocusPoint(float InExistingFocus, float InNewFocus,  float InputTolerance)
+{
+	LensDataTableUtils::ChangeFocusPoint(FocusPoints, InExistingFocus, InNewFocus, InputTolerance);
+}
+
+void FFocalLengthTable::MergeFocusPoint(float InSrcFocus, float InDestFocus, bool bReplaceExistingZoomPoints,  float InputTolerance)
+{
+	LensDataTableUtils::MergeFocusPoint(FocusPoints, InSrcFocus, InDestFocus, bReplaceExistingZoomPoints, InputTolerance);
+}
+
 void FFocalLengthTable::RemoveZoomPoint(float InFocus, float InZoom)
 {
 	LensDataTableUtils::RemoveZoomPoint(FocusPoints, InFocus, InZoom);
 }
 
+bool FFocalLengthTable::HasZoomPoint(float InFocus, float InZoom,  float InputTolerance)
+{
+	return DoesZoomPointExists(InFocus, InZoom, InputTolerance);
+}
+
+void FFocalLengthTable::ChangeZoomPoint(float InFocus, float InExistingZoom, float InNewZoom,  float InputTolerance)
+{
+	LensDataTableUtils::ChangeZoomPoint(FocusPoints, InFocus, InExistingZoom, InNewZoom, InputTolerance);
+}
+
 bool FFocalLengthTable::DoesFocusPointExists(float InFocus, float InputTolerance) const
 {
 	if (GetFocusPoint(InFocus, InputTolerance) != nullptr)
@@ -235,11 +273,8 @@ bool FFocalLengthTable::GetPoint(const float InFocus, const float InZoom, FFocal
 {
 	if (const FFocalLengthFocusPoint* FocalLengthFocusPoint = GetFocusPoint(InFocus, InputTolerance))
 	{
-		FFocalLengthZoomPoint ZoomPont;
-		if (FocalLengthFocusPoint->GetPoint(InZoom, ZoomPont, InputTolerance))
+		if (FocalLengthFocusPoint->GetPoint(InZoom, OutData, InputTolerance))
 		{
-			// Copy struct to outer
-			OutData = ZoomPont.FocalLengthInfo;
 			return true;
 		}
 	}
diff --git a/Engine/Plugins/VirtualProduction/CameraCalibrationCore/Source/CameraCalibrationCore/Private/Tables/ImageCenterTable.cpp b/Engine/Plugins/VirtualProduction/CameraCalibrationCore/Source/CameraCalibrationCore/Private/Tables/ImageCenterTable.cpp
--- a/Engine/Plugins/VirtualProduction/CameraCalibrationCore/Source/CameraCalibrationCore/Private/Tables/ImageCenterTable.cpp
+++ b/Engine/Plugins/VirtualProduction/CameraCalibrationCore/Source/CameraCalibrationCore/Private/Tables/ImageCenterTable.cpp
@@ -197,11 +197,36 @@ void FImageCenterTable::RemoveFocusPoint(float InFocus)
 	LensDataTableUtils::RemoveFocusPoint(FocusPoints, InFocus);
 }
 
+bool FImageCenterTable::HasFocusPoint(float InFocus, float InputTolerance) const
+{
+	return DoesFocusPointExists(InFocus, InputTolerance);
+}
+
+void FImageCenterTable::ChangeFocusPoint(float InExistingFocus, float InNewFocus, float InputTolerance)
+{
+	LensDataTableUtils::ChangeFocusPoint(FocusPoints, InExistingFocus, InNewFocus, InputTolerance);
+}
+
+void FImageCenterTable::MergeFocusPoint(float InSrcFocus, float InDestFocus, bool bReplaceExistingZoomPoints, float InputTolerance)
+{
+	LensDataTableUtils::MergeFocusPoint(FocusPoints, InSrcFocus, InDestFocus, bReplaceExistingZoomPoints, InputTolerance);
+}
+
 void FImageCenterTable::RemoveZoomPoint(float InFocus, float InZoom)
 {
 	LensDataTableUtils::RemoveZoomPoint(FocusPoints, InFocus, InZoom);
 }
 
+bool FImageCenterTable::HasZoomPoint(float InFocus, float InZoom, float InputTolerance)
+{
+	return DoesZoomPointExists(InFocus, InZoom, InputTolerance);
+}
+
+void FImageCenterTable::ChangeZoomPoint(float InFocus, float InExistingZoom, float InNewZoom, float InputTolerance)
+{
+	LensDataTableUtils::ChangeZoomPoint(FocusPoints, InFocus, InExistingZoom, InNewZoom, InputTolerance);
+}
+
 bool FImageCenterTable::AddPoint(float InFocus, float InZoom, const FImageCenterInfo& InData, float InputTolerance,
 	bool bIsCalibrationPoint)
 {
diff --git a/Engine/Plugins/VirtualProduction/CameraCalibrationCore/Source/CameraCalibrationCore/Private/Tables/LensTableUtils.h b/Engine/Plugins/VirtualProduction/CameraCalibrationCore/Source/CameraCalibrationCore/Private/Tables/LensTableUtils.h
--- a/Engine/Plugins/VirtualProduction/CameraCalibrationCore/Source/CameraCalibrationCore/Private/Tables/LensTableUtils.h
+++ b/Engine/Plugins/VirtualProduction/CameraCalibrationCore/Source/CameraCalibrationCore/Private/Tables/LensTableUtils.h
@@ -21,6 +21,74 @@ namespace LensDataTableUtils
     	}
 	}
 
+	/** Changes the value of a focus point in the container */
+	template<typename FocusPointType>
+	void ChangeFocusPoint(TArray<FocusPointType>& Container, float InExistingFocus, float InNewFocus, float InputTolerance = KINDA_SMALL_NUMBER)
+	{
+		const int32 FoundIndex = Container.IndexOfByPredicate([InExistingFocus, InputTolerance](const FocusPointType& Point)
+		{
+			return FMath::IsNearlyEqual(Point.Focus, InExistingFocus, InputTolerance);
+		});
+		
+		if (FoundIndex != INDEX_NONE)
+		{
+			Container[FoundIndex].Focus = InNewFocus;
+		}
+	}
+
+	/** Merges the points in the specified source focus into the specified destination focus */
+	template<typename FocusPointType>
+	void MergeFocusPoint(TArray<FocusPointType>& Container, float InSrcFocus, float InDestFocus, bool bReplaceExistingZoomPoints, float InputTolerance = KINDA_SMALL_NUMBER)
+	{
+		const int32 SrcIndex = Container.IndexOfByPredicate([InSrcFocus, InputTolerance](const FocusPointType& Point)
+		{
+			return FMath::IsNearlyEqual(Point.Focus, InSrcFocus, InputTolerance);
+		});
+		
+		const int32 DestIndex = Container.IndexOfByPredicate([InDestFocus, InputTolerance](const FocusPointType& Point)
+		{
+			return FMath::IsNearlyEqual(Point.Focus, InDestFocus, InputTolerance);
+		});
+		
+		if (SrcIndex != INDEX_NONE)
+		{
+			if (DestIndex != INDEX_NONE)
+			{
+				FocusPointType& SrcPoint = Container[SrcIndex];
+				FocusPointType& DestPoint = Container[DestIndex];
+
+				for (int32 Index = 0; Index < SrcPoint.GetNumPoints(); ++Index)
+				{
+					float Zoom = SrcPoint.GetZoom(Index);
+
+					typename FocusPointType::PointType SrcData;
+					SrcPoint.GetPoint(Zoom, SrcData);
+					bool bIsCalibrationPoint = SrcPoint.IsCalibrationPoint(Zoom);
+
+					typename FocusPointType::PointType DestData;
+					if (DestPoint.GetPoint(Zoom, DestData))
+					{
+						if (!bReplaceExistingZoomPoints)
+						{
+							continue;
+						}
+
+						DestPoint.RemovePoint(Zoom);
+					}
+					
+					DestPoint.AddPoint(Zoom, SrcData, InputTolerance, bIsCalibrationPoint);
+				}
+
+				RemoveFocusPoint(Container, InSrcFocus);
+			}
+			else
+			{
+				// The destination doesn't exist, so we can just change the source focus point to the destination focus point
+				ChangeFocusPoint(Container, InSrcFocus, InDestFocus);
+			}
+		}
+	}
+
 	/** Gets all point info for specific data table */
 	template<typename FPointInfoType, typename FDataTableType>
 	TArray<FPointInfoType> GetAllPointsInfo(const FDataTableType& InTable)
@@ -69,6 +137,27 @@ namespace LensDataTableUtils
 		}
 	}
 
+	/** Changes the value of a zoom point for a given focus value in a container */
+	template<typename FocusPointType>
+	void ChangeZoomPoint(TArray<FocusPointType>& Container, float InFocus, float InExistingZoom, float InNewZoom, float InputTolerance = KINDA_SMALL_NUMBER)
+	{
+		const int32 FoundIndex = Container.IndexOfByPredicate([InFocus, InputTolerance](const FocusPointType& Point)
+		{
+			return FMath::IsNearlyEqual(Point.Focus, InFocus, InputTolerance);
+		});
+		
+		if (FoundIndex != INDEX_NONE)
+		{
+			typename FocusPointType::PointType PointData;
+			Container[FoundIndex].GetPoint(InExistingZoom, PointData);
+			
+			bool bIsCalibrationPoint = Container[FoundIndex].IsCalibrationPoint(InExistingZoom);
+			
+			Container[FoundIndex].RemovePoint(InExistingZoom);
+			Container[FoundIndex].AddPoint(InNewZoom, PointData, InputTolerance, bIsCalibrationPoint);
+		}
+	}
+	
 	/** Adds a point at a specified focus and zoom input values */
 	template<typename FocusPointType, typename DataType>
 	bool AddPoint(TArray<FocusPointType>& InContainer, float InFocus, float InZoom, const DataType& InData, float InputTolerance, bool bIsCalibrationPoint)
diff --git a/Engine/Plugins/VirtualProduction/CameraCalibrationCore/Source/CameraCalibrationCore/Private/Tables/NodalOffsetTable.cpp b/Engine/Plugins/VirtualProduction/CameraCalibrationCore/Source/CameraCalibrationCore/Private/Tables/NodalOffsetTable.cpp
--- a/Engine/Plugins/VirtualProduction/CameraCalibrationCore/Source/CameraCalibrationCore/Private/Tables/NodalOffsetTable.cpp
+++ b/Engine/Plugins/VirtualProduction/CameraCalibrationCore/Source/CameraCalibrationCore/Private/Tables/NodalOffsetTable.cpp
@@ -239,11 +239,36 @@ void FNodalOffsetTable::RemoveFocusPoint(float InFocus)
 	LensDataTableUtils::RemoveFocusPoint(FocusPoints, InFocus);
 }
 
+bool FNodalOffsetTable::HasFocusPoint(float InFocus, float InputTolerance) const
+{
+	return DoesFocusPointExists(InFocus, InputTolerance);
+}
+
+void FNodalOffsetTable::ChangeFocusPoint(float InExistingFocus, float InNewFocus, float InputTolerance)
+{
+	LensDataTableUtils::ChangeFocusPoint(FocusPoints, InExistingFocus, InNewFocus, InputTolerance);
+}
+
+void FNodalOffsetTable::MergeFocusPoint(float InSrcFocus, float InDestFocus, bool bReplaceExistingZoomPoints, float InputTolerance)
+{
+	LensDataTableUtils::MergeFocusPoint(FocusPoints, InSrcFocus, InDestFocus, bReplaceExistingZoomPoints, InputTolerance);
+}
+
 void FNodalOffsetTable::RemoveZoomPoint(float InFocus, float InZoom)
 {
 	LensDataTableUtils::RemoveZoomPoint(FocusPoints, InFocus, InZoom);
 }
 
+bool FNodalOffsetTable::HasZoomPoint(float InFocus, float InZoom, float InputTolerance)
+{
+	return DoesZoomPointExists(InFocus, InZoom, InputTolerance);
+}
+
+void FNodalOffsetTable::ChangeZoomPoint(float InFocus, float InExistingZoom, float InNewZoom, float InputTolerance)
+{
+	LensDataTableUtils::ChangeZoomPoint(FocusPoints, InFocus, InExistingZoom, InNewZoom, InputTolerance);
+}
+
 bool FNodalOffsetTable::DoesFocusPointExists(float InFocus, float InputTolerance) const
 {
 	if (GetFocusPoint(InFocus, InputTolerance) != nullptr)
diff --git a/Engine/Plugins/VirtualProduction/CameraCalibrationCore/Source/CameraCalibrationCore/Private/Tables/STMapTable.cpp b/Engine/Plugins/VirtualProduction/CameraCalibrationCore/Source/CameraCalibrationCore/Private/Tables/STMapTable.cpp
--- a/Engine/Plugins/VirtualProduction/CameraCalibrationCore/Source/CameraCalibrationCore/Private/Tables/STMapTable.cpp
+++ b/Engine/Plugins/VirtualProduction/CameraCalibrationCore/Source/CameraCalibrationCore/Private/Tables/STMapTable.cpp
@@ -90,6 +90,21 @@ bool FSTMapFocusPoint::SetPoint(float InZoom, const FSTMapInfo& InData, float In
 	return false;
 }
 
+bool FSTMapFocusPoint::IsCalibrationPoint(float InZoom, float InputTolerance)
+{
+	const FKeyHandle Handle = MapBlendingCurve.FindKey(InZoom, InputTolerance);
+	if(Handle != FKeyHandle::Invalid())
+	{
+		const int32 PointIndex = MapBlendingCurve.GetIndexSafe(Handle);
+		if(ensure(ZoomPoints.IsValidIndex(PointIndex)))
+		{
+			return ZoomPoints[PointIndex].bIsCalibrationPoint;
+		}
+	}
+
+	return false;
+}
+
 void FSTMapFocusPoint::RemovePoint(float InZoomValue)
 {
 	const int32 FoundIndex = ZoomPoints.IndexOfByPredicate([InZoomValue](const FSTMapZoomPoint& Point) { return FMath::IsNearlyEqual(Point.Zoom, InZoomValue); });
@@ -195,11 +210,36 @@ void FSTMapTable::RemoveFocusPoint(float InFocus)
 	LensDataTableUtils::RemoveFocusPoint(FocusPoints, InFocus);
 }
 
+bool FSTMapTable::HasFocusPoint(float InFocus, float InputTolerance) const
+{
+	return DoesFocusPointExists(InFocus, InputTolerance);
+}
+
+void FSTMapTable::ChangeFocusPoint(float InExistingFocus, float InNewFocus, float InputTolerance)
+{
+	LensDataTableUtils::ChangeFocusPoint(FocusPoints, InExistingFocus, InNewFocus, InputTolerance);
+}
+
+void FSTMapTable::MergeFocusPoint(float InSrcFocus, float InDestFocus, bool bReplaceExistingZoomPoints, float InputTolerance)
+{
+	LensDataTableUtils::MergeFocusPoint(FocusPoints, InSrcFocus, InDestFocus, bReplaceExistingZoomPoints, InputTolerance);
+}
+
 void FSTMapTable::RemoveZoomPoint(float InFocus, float InZoom)
 {
 	LensDataTableUtils::RemoveZoomPoint(FocusPoints, InFocus, InZoom);
 }
 
+bool FSTMapTable::HasZoomPoint(float InFocus, float InZoom, float InputTolerance)
+{
+	return DoesZoomPointExists(InFocus, InZoom, InputTolerance);
+}
+
+void FSTMapTable::ChangeZoomPoint(float InFocus, float InExistingZoom, float InNewZoom, float InputTolerance)
+{
+	LensDataTableUtils::ChangeZoomPoint(FocusPoints, InFocus, InExistingZoom, InNewZoom, InputTolerance);
+}
+
 TMap<ELensDataCategory, FLinkPointMetadata> FSTMapTable::GetLinkedCategories() const
 {
 	static TMap<ELensDataCategory, FLinkPointMetadata> LinkedToCategories =
diff --git a/Engine/Plugins/VirtualProduction/CameraCalibrationCore/Source/CameraCalibrationCore/Public/LensFile.h b/Engine/Plugins/VirtualProduction/CameraCalibrationCore/Source/CameraCalibrationCore/Public/LensFile.h
--- a/Engine/Plugins/VirtualProduction/CameraCalibrationCore/Source/CameraCalibrationCore/Public/LensFile.h
+++ b/Engine/Plugins/VirtualProduction/CameraCalibrationCore/Source/CameraCalibrationCore/Public/LensFile.h
@@ -190,10 +190,30 @@ public:
 	UFUNCTION(BlueprintCallable, Category = "Lens Table")
 	void RemoveFocusPoint(ELensDataCategory InDataCategory, float InFocus);
 
+	/** Checks to see if there is a focal point for the specified focus in the data category */
+	UFUNCTION(BlueprintCallable, Category = "Lens Table")
+	bool HasFocusPoint(ELensDataCategory InDataCategory, float InFocus) const;
+
+	/** Changes the value of a focus point */
+	UFUNCTION(BlueprintCallable, Category = "Lens Table")
+	void ChangeFocusPoint(ELensDataCategory InDataCategory, float InExistingFocus, float InNewFocus);
+
+	/** Merges the contents of one focus point into another focus point */
+	UFUNCTION(BlueprintCallable, Category = "Lens Table")
+	void MergeFocusPoint(ELensDataCategory InDataCategory, float InSrcFocus, float InDestFocus, bool bReplaceExistingZoomPoints);
+
 	/** Removes a zoom point */
 	UFUNCTION(BlueprintCallable, Category = "Lens Table")
 	void RemoveZoomPoint(ELensDataCategory InDataCategory, float InFocus, float InZoom);
 
+	/** Removes a zoom point */
+	UFUNCTION(BlueprintCallable, Category = "Lens Table")
+	bool HasZoomPoint(ELensDataCategory InDataCategory, float InFocus, float InZoom);
+
+	/** Changes the value of a zoom point */
+	UFUNCTION(BlueprintCallable, Category = "Lens Table")
+	void ChangeZoomPoint(ELensDataCategory InDataCategory, float InFocus, float InExistingZoom, float InNewZoom);
+	
 	/** Removes all points of all tables */
 	UFUNCTION(BlueprintCallable, Category = "Lens Table")
 	void ClearAll();
diff --git a/Engine/Plugins/VirtualProduction/CameraCalibrationCore/Source/CameraCalibrationCore/Public/Tables/DistortionParametersTable.h b/Engine/Plugins/VirtualProduction/CameraCalibrationCore/Source/CameraCalibrationCore/Public/Tables/DistortionParametersTable.h
--- a/Engine/Plugins/VirtualProduction/CameraCalibrationCore/Source/CameraCalibrationCore/Public/Tables/DistortionParametersTable.h
+++ b/Engine/Plugins/VirtualProduction/CameraCalibrationCore/Source/CameraCalibrationCore/Public/Tables/DistortionParametersTable.h
@@ -35,6 +35,9 @@ USTRUCT()
 struct CAMERACALIBRATIONCORE_API FDistortionFocusPoint : public FBaseFocusPoint
 {
 	GENERATED_BODY()
+
+	using PointType = FDistortionInfo;
+	
 public:
 	//~ Begin FBaseFocusPoint Interface
 	virtual float GetFocus() const override { return Focus; }
@@ -52,7 +55,10 @@ public:
 
 	/** Sets an existing point at InZoom. Updates existing one if tolerance is met */
 	bool SetPoint(float InZoom, const FDistortionInfo& InData, float InputTolerance = KINDA_SMALL_NUMBER);
-
+	
+	/** Gets whether the point at InZoom is a calibration point. */
+	bool IsCalibrationPoint(float InZoom, float InputTolerance = KINDA_SMALL_NUMBER) { return false; }
+	
 	/** Returns true if this point is empty */
 	bool IsEmpty() const;
 	
@@ -120,9 +126,24 @@ public:
 	/** Removes a focus point identified as InFocusIdentifier */
 	void RemoveFocusPoint(float InFocus);
 
+	/** Checks to see if there exists a focus point matching the specified focus value */
+	bool HasFocusPoint(float InFocus, float InputTolerance = KINDA_SMALL_NUMBER) const;
+	
+	/** Changes the value of a focus point */
+	void ChangeFocusPoint(float InExistingFocus, float InNewFocus, float InputTolerance = KINDA_SMALL_NUMBER);
+
+	/** Merges the points in the specified source focus into the specified destination focus */
+	void MergeFocusPoint(float InSrcFocus, float InDestFocus, bool bReplaceExistingZoomPoints, float InputTolerance = KINDA_SMALL_NUMBER);
+	
 	/** Removes a zoom point from a focus point*/
 	void RemoveZoomPoint(float InFocus, float InZoom);
 
+	/** Checks to see if there exists a zoom point matching the specified zoom and focus values */
+	bool HasZoomPoint(float InFocus, float InZoom, float InputTolerance = KINDA_SMALL_NUMBER);
+
+	/** Changes the value of a zoom point */
+	void ChangeZoomPoint(float InFocus, float InExistingZoom, float InNewZoom, float InputTolerance = KINDA_SMALL_NUMBER);
+	
 	/** Adds a new point in the table */
 	bool AddPoint(float InFocus, float InZoom, const FDistortionInfo& InData, float InputTolerance, bool bIsCalibrationPoint);
 
diff --git a/Engine/Plugins/VirtualProduction/CameraCalibrationCore/Source/CameraCalibrationCore/Public/Tables/FocalLengthTable.h b/Engine/Plugins/VirtualProduction/CameraCalibrationCore/Source/CameraCalibrationCore/Public/Tables/FocalLengthTable.h
--- a/Engine/Plugins/VirtualProduction/CameraCalibrationCore/Source/CameraCalibrationCore/Public/Tables/FocalLengthTable.h
+++ b/Engine/Plugins/VirtualProduction/CameraCalibrationCore/Source/CameraCalibrationCore/Public/Tables/FocalLengthTable.h
@@ -39,7 +39,9 @@ USTRUCT()
 struct CAMERACALIBRATIONCORE_API FFocalLengthFocusPoint : public FBaseFocusPoint
 {
 	GENERATED_BODY()
-	
+
+	using PointType = FFocalLengthInfo;
+
 public:
 	//~ Begin FBaseFocusPoint Interface
 	virtual float GetFocus() const override { return Focus; }
@@ -47,8 +49,8 @@ public:
 	virtual float GetZoom(int32 Index) const override;
 	//~ End FBaseFocusPoint Interface
 
-	/** Returns zoom value for a given float */
-	bool GetPoint(float InZoom, FFocalLengthZoomPoint& OutZoomPont, float InputTolerance = KINDA_SMALL_NUMBER) const;
+	/** Returns data type copy value for a given float  */
+	bool GetPoint(float InZoom, FFocalLengthInfo& OutData, float InputTolerance = KINDA_SMALL_NUMBER) const;
 
 	/** Adds a new point at InZoom. Updates existing one if tolerance is met */
 	bool AddPoint(float InZoom, const FFocalLengthInfo& InData, float InputTolerance, bool bIsCalibrationPoint);
@@ -56,6 +58,9 @@ public:
 	/** Sets an existing point at InZoom. Updates existing one if tolerance is met */
 	bool SetPoint(float InZoom, const FFocalLengthInfo& InData, float InputTolerance = KINDA_SMALL_NUMBER);
 
+	/** Gets whether the point at InZoom is a calibration point. */
+	bool IsCalibrationPoint(float InZoom, float InputTolerance = KINDA_SMALL_NUMBER);
+
 	/** Returns data at the requested index */
 	bool GetValue(int32 Index, FFocalLengthInfo& OutData) const;
 
@@ -128,11 +133,26 @@ public:
 	/** Removes a focus point identified as InFocusIdentifier */
 	void RemoveFocusPoint(float InFocus);
 
+	/** Checks to see if there exists a focus point matching the specified focus value */
+	bool HasFocusPoint(float InFocus, float InputTolerance = KINDA_SMALL_NUMBER) const;
+
+	/** Changes the value of a focus point */
+	void ChangeFocusPoint(float InExistingFocus, float InNewFocus, float InputTolerance = KINDA_SMALL_NUMBER);
+
+	/** Merges the points in the specified source focus into the specified destination focus */
+	void MergeFocusPoint(float InSrcFocus, float InDestFocus, bool bReplaceExistingZoomPoints, float InputTolerance = KINDA_SMALL_NUMBER);
+
 	/** Removes a zoom point from a focus point*/
 	void RemoveZoomPoint(float InFocus, float InZoom);
 
+	/** Checks to see if there exists a zoom point matching the specified zoom and focus values */
+	bool HasZoomPoint(float InFocus, float InZoom, float InputTolerance = KINDA_SMALL_NUMBER);
+
+	/** Changes the value of a zoom point */
+	void ChangeZoomPoint(float InFocus, float InExistingZoom, float InNewZoom, float InputTolerance = KINDA_SMALL_NUMBER);
+	
 	/** Adds a new point in the table */
-	bool AddPoint(float InFocus, float InZoom, const FFocalLengthInfo& InData,  float InputTolerance, bool bIsCalibrationPoint);
+	bool AddPoint(float InFocus, float InZoom, const FFocalLengthInfo& InData, float InputTolerance, bool bIsCalibrationPoint);
 
 	/** Get the point from the table */
 	bool GetPoint(const float InFocus, const float InZoom, FFocalLengthInfo& OutData, float InputTolerance = KINDA_SMALL_NUMBER) const;
@@ -140,7 +160,6 @@ public:
 	/** Set a new point into the table */
 	bool SetPoint(float InFocus, float InZoom, const FFocalLengthInfo& InData, float InputTolerance = KINDA_SMALL_NUMBER);
 
-
 public:
 
 	/** Lists of focus points */
diff --git a/Engine/Plugins/VirtualProduction/CameraCalibrationCore/Source/CameraCalibrationCore/Public/Tables/ImageCenterTable.h b/Engine/Plugins/VirtualProduction/CameraCalibrationCore/Source/CameraCalibrationCore/Public/Tables/ImageCenterTable.h
--- a/Engine/Plugins/VirtualProduction/CameraCalibrationCore/Source/CameraCalibrationCore/Public/Tables/ImageCenterTable.h
+++ b/Engine/Plugins/VirtualProduction/CameraCalibrationCore/Source/CameraCalibrationCore/Public/Tables/ImageCenterTable.h
@@ -17,6 +17,8 @@ struct CAMERACALIBRATIONCORE_API FImageCenterFocusPoint : public FBaseFocusPoint
 {
 	GENERATED_BODY()
 
+	using PointType = FImageCenterInfo;
+	
 public:
 	//~ Begin FBaseFocusPoint Interface
 	virtual float GetFocus() const override { return Focus; }
@@ -32,7 +34,10 @@ public:
 
 	/** Sets an existing point at InZoom. Updates existing one if tolerance is met */
 	bool SetPoint(float InZoom, const FImageCenterInfo& InData, float InputTolerance = KINDA_SMALL_NUMBER);
-
+	
+	/** Gets whether the point at InZoom is a calibration point. */
+	bool IsCalibrationPoint(float InZoom, float InputTolerance = KINDA_SMALL_NUMBER) { return false; }
+	
 	/** Removes a point corresponding to specified zoom */
 	void RemovePoint(float InZoomValue);
 
@@ -101,9 +106,24 @@ public:
 	/** Removes a focus point identified as InFocusIdentifier */
 	void RemoveFocusPoint(float InFocus);
 
+	/** Checks to see if there exists a focus point matching the specified focus value */
+	bool HasFocusPoint(float InFocus, float InputTolerance = KINDA_SMALL_NUMBER) const;
+
+	/** Changes the value of a focus point */
+	void ChangeFocusPoint(float InExistingFocus, float InNewFocus, float InputTolerance = KINDA_SMALL_NUMBER);
+
+	/** Merges the points in the specified source focus into the specified destination focus */
+	void MergeFocusPoint(float InSrcFocus, float InDestFocus, bool bReplaceExistingZoomPoints, float InputTolerance = KINDA_SMALL_NUMBER);
+
 	/** Removes a zoom point from a focus point*/
 	void RemoveZoomPoint(float InFocus, float InZoom);
+	
+	/** Checks to see if there exists a zoom point matching the specified zoom and focus values */
+	bool HasZoomPoint(float InFocus, float InZoom, float InputTolerance = KINDA_SMALL_NUMBER);
 
+	/** Changes the value of a zoom point */
+	void ChangeZoomPoint(float InFocus, float InExistingZoom, float InNewZoom, float InputTolerance = KINDA_SMALL_NUMBER);
+	
 	/** Adds a new point in the table */
 	bool AddPoint(float InFocus, float InZoom, const FImageCenterInfo& InData, float InputTolerance, bool bIsCalibrationPoint);
 
diff --git a/Engine/Plugins/VirtualProduction/CameraCalibrationCore/Source/CameraCalibrationCore/Public/Tables/NodalOffsetTable.h b/Engine/Plugins/VirtualProduction/CameraCalibrationCore/Source/CameraCalibrationCore/Public/Tables/NodalOffsetTable.h
--- a/Engine/Plugins/VirtualProduction/CameraCalibrationCore/Source/CameraCalibrationCore/Public/Tables/NodalOffsetTable.h
+++ b/Engine/Plugins/VirtualProduction/CameraCalibrationCore/Source/CameraCalibrationCore/Public/Tables/NodalOffsetTable.h
@@ -17,6 +17,8 @@ struct CAMERACALIBRATIONCORE_API FNodalOffsetFocusPoint : public FBaseFocusPoint
 {
 	GENERATED_BODY()
 
+	using PointType = FNodalPointOffset;
+	
 public:
 	//~ Begin FBaseFocusPoint Interface
 	virtual float GetFocus() const override { return Focus; }
@@ -33,6 +35,9 @@ public:
 	/** Sets an existing point at InZoom. Updates existing one if tolerance is met */
 	bool SetPoint(float InZoom, const FNodalPointOffset& InData, float InputTolerance = KINDA_SMALL_NUMBER);
 	
+	/** Gets whether the point at InZoom is a calibration point. */
+	bool IsCalibrationPoint(float InZoom, float InputTolerance = KINDA_SMALL_NUMBER) { return false; }
+	
 	/** Removes a point corresponding to specified zoom */
 	void RemovePoint(float InZoomValue);
 
@@ -107,9 +112,24 @@ public:
 	/** Removes a focus point identified as InFocusIdentifier */
 	void RemoveFocusPoint(float InFocus);
 
+	/** Checks to see if there exists a focus point matching the specified focus value */
+	bool HasFocusPoint(float InFocus, float InputTolerance = KINDA_SMALL_NUMBER) const;
+
+	/** Changes the value of a focus point */
+	void ChangeFocusPoint(float InExistingFocus, float InNewFocus, float InputTolerance = KINDA_SMALL_NUMBER);
+
+	/** Merges the points in the specified source focus into the specified destination focus */
+	void MergeFocusPoint(float InSrcFocus, float InDestFocus, bool bReplaceExistingZoomPoints, float InputTolerance = KINDA_SMALL_NUMBER);
+	
 	/** Removes a zoom point from a focus point*/
 	void RemoveZoomPoint(float InFocus, float InZoom);
 
+	/** Checks to see if there exists a zoom point matching the specified zoom and focus values */
+	bool HasZoomPoint(float InFocus, float InZoom, float InputTolerance = KINDA_SMALL_NUMBER);
+
+	/** Changes the value of a zoom point */
+	void ChangeZoomPoint(float InFocus, float InExistingZoom, float InNewZoom, float InputTolerance = KINDA_SMALL_NUMBER);
+	
 	/** Adds a new point in the table */
 	bool AddPoint(float InFocus, float InZoom, const FNodalPointOffset& InData,  float InputTolerance, bool bIsCalibrationPoint);
 
diff --git a/Engine/Plugins/VirtualProduction/CameraCalibrationCore/Source/CameraCalibrationCore/Public/Tables/STMapTable.h b/Engine/Plugins/VirtualProduction/CameraCalibrationCore/Source/CameraCalibrationCore/Public/Tables/STMapTable.h
--- a/Engine/Plugins/VirtualProduction/CameraCalibrationCore/Source/CameraCalibrationCore/Public/Tables/STMapTable.h
+++ b/Engine/Plugins/VirtualProduction/CameraCalibrationCore/Source/CameraCalibrationCore/Public/Tables/STMapTable.h
@@ -69,6 +69,8 @@ struct CAMERACALIBRATIONCORE_API FSTMapFocusPoint : public FBaseFocusPoint
 {
 	GENERATED_BODY()
 
+	using PointType = FSTMapInfo;
+	
 public:
 	//~ Begin FBaseFocusPoint Interface
 	virtual float GetFocus() const override { return Focus; }
@@ -91,6 +93,9 @@ public:
 	/** Sets an existing point at InZoom. Updates existing one if tolerance is met */
 	bool SetPoint(float InZoom, const FSTMapInfo& InData, float InputTolerance = KINDA_SMALL_NUMBER);
 	
+	/** Gets whether the point at InZoom is a calibration point. */
+	bool IsCalibrationPoint(float InZoom, float InputTolerance = KINDA_SMALL_NUMBER);
+	
 	/** Removes a point corresponding to specified zoom */
 	void RemovePoint(float InZoomValue);
 
@@ -159,9 +164,24 @@ public:
 	/** Removes a focus point identified as InFocusIdentifier */
 	void RemoveFocusPoint(float InFocus);
 
+	/** Checks to see if there exists a focus point matching the specified focus value */
+	bool HasFocusPoint(float InFocus, float InputTolerance = KINDA_SMALL_NUMBER) const;
+
+	/** Changes the value of a focus point */
+	void ChangeFocusPoint(float InExistingFocus, float InNewFocus, float InputTolerance = KINDA_SMALL_NUMBER);
+
+	/** Merges the points in the specified source focus into the specified destination focus */
+	void MergeFocusPoint(float InSrcFocus, float InDestFocus, bool bReplaceExistingZoomPoints, float InputTolerance = KINDA_SMALL_NUMBER);
+	
 	/** Removes a zoom point from a focus point*/
 	void RemoveZoomPoint(float InFocus, float InZoom);
 
+	/** Checks to see if there exists a zoom point matching the specified zoom and focus values */
+	bool HasZoomPoint(float InFocus, float InZoom, float InputTolerance = KINDA_SMALL_NUMBER);
+
+	/** Changes the value of a zoom point */
+	void ChangeZoomPoint(float InFocus, float InExistingZoom, float InNewZoom, float InputTolerance = KINDA_SMALL_NUMBER);
+	
 	/** Adds a new point in the table */
 	bool AddPoint(float InFocus, float InZoom, const FSTMapInfo& InData,  float InputTolerance, bool bIsCalibrationPoint);
 
