CameraCalibration: Add FixAspectRatio solver flag to force the solver to respect a given pixel aspect ratio. For spherical lenses, the pixel aspect ratio is 1.0, which means the solver will enforce that Fx must equal Fy when solving for focal length.
- Add a cvar to the distortion checkerboard algo to control whether the 3D checkerboard points should be used, or if they should be replaced by dummy 3D coordinates.
- Add an array of camera transforms to the FDistortionCalibrationResult struct to let the solver output the camera pose for each input image.
- Add a cvar to the solver to constrain it to solve for only one camera pose (using the input tracked poses to obtain the pose for each image).
  - The constraint itself is not new, but it is tied to a dedicated cvar instead of being lumped into the UseExtrinsicGuess solver flag.

#rnx

diff --git a/Engine/Plugins/VirtualProduction/CameraCalibration/Source/CameraCalibrationEditor/Private/Calibrators/CameraCalibrationSolver.cpp b/Engine/Plugins/VirtualProduction/CameraCalibration/Source/CameraCalibrationEditor/Private/Calibrators/CameraCalibrationSolver.cpp
--- a/Engine/Plugins/VirtualProduction/CameraCalibration/Source/CameraCalibrationEditor/Private/Calibrators/CameraCalibrationSolver.cpp
+++ b/Engine/Plugins/VirtualProduction/CameraCalibration/Source/CameraCalibrationEditor/Private/Calibrators/CameraCalibrationSolver.cpp
@@ -18,6 +18,7 @@
 #define LOCTEXT_NAMESPACE "CameraCalibrationSolver"
 
 static TAutoConsoleVariable<int> CVarUseLegacySphericalSolver(TEXT("CameraCalibration.UseLegacySphericalSolver"), 0, TEXT("If set, the legacy OpenCV spherical solver will be used"));
+static TAutoConsoleVariable<bool> CVarConstrainWithCameraMotion(TEXT("CameraCalibration.ConstrainWithCameraMotion"), false, TEXT("If true, the solver will only solve for the camera pose of the first image, and will use the input tracked camera poses to transform the solve pose for subsequent images."));
 static TAutoConsoleVariable<float> CVarRotationStepValue(TEXT("CameraCalibration.RotationStepValue"), 0.05, TEXT("The value of the initial step size to use when finding an optimal nodal offset rotation that minimizes reprojection error."));
 static TAutoConsoleVariable<float> CVarLocationStepValue(TEXT("CameraCalibration.LocationStepValue"), 0.5, TEXT("The value of the initial step size to use when finding an optimal nodal offset location that minimizes reprojection error."));
 
@@ -162,6 +163,10 @@ FDistortionCalibrationResult ULensDistortionSolverOpenCV::Solve_Implementation(
 		{
 			LegacyFlags |= cv::CALIB_FIX_PRINCIPAL_POINT;
 		}
+		if (EnumHasAnyFlags(SolverFlags, ECalibrationFlags::FixAspectRatio))
+		{
+			LegacyFlags |= cv::CALIB_FIX_ASPECT_RATIO;
+		}
 
 		RMSE = cv::calibrateCamera(
 			Samples3d,
@@ -192,6 +197,15 @@ FDistortionCalibrationResult ULensDistortionSolverOpenCV::Solve_Implementation(
 
 		Result.ReprojectionError = RMSE;
 
+		Result.CameraPoses.Reserve(NumImages);
+		for (int32 ImageIndex = 0; ImageIndex < NumImages; ++ImageIndex)
+		{
+			FTransform LegacyCameraPose;
+			FOpenCVHelper::MakeCameraPoseFromObjectVectors(Rvecs[ImageIndex], Tvecs[ImageIndex], LegacyCameraPose);
+
+			Result.CameraPoses.Add(LegacyCameraPose);
+		}
+
 		return Result;
 	}
 
@@ -207,8 +221,10 @@ FDistortionCalibrationResult ULensDistortionSolverOpenCV::Solve_Implementation(
 
 	// If using a guess for the extrinsic parameters, then the solver will be constrained to solve only one camera pose, and needs only one set of extrinsic parameters.
 	// Otherwise, the solver needs one set of extrinsic parameters per image.
+	const bool bConstrainWithCameraMotion = CVarConstrainWithCameraMotion.GetValueOnAnyThread();
+
 	int NumPosesToSolve = 0;
-	if ((EnumHasAnyFlags(SolverFlags, ECalibrationFlags::UseExtrinsicGuess)))
+	if (bConstrainWithCameraMotion)
 	{
 		NumPosesToSolve = 1;
 	}
@@ -235,6 +251,10 @@ FDistortionCalibrationResult ULensDistortionSolverOpenCV::Solve_Implementation(
 
 	// Instruct the solver to ignore some parameters when running its solve
 	uchar* Mask = Solver.Mask.ptr<uchar>();
+	if (EnumHasAnyFlags(SolverFlags, ECalibrationFlags::FixAspectRatio))
+	{
+		Mask[0] = 0;
+	}
 	if (EnumHasAnyFlags(SolverFlags, ECalibrationFlags::FixFocalLength))
 	{
 		Mask[0] = 0;
@@ -268,17 +288,20 @@ FDistortionCalibrationResult ULensDistortionSolverOpenCV::Solve_Implementation(
 	// Initialize the starting guess for the camera's extrinsic parameters. 
 	// If using a guess for the extrinsic parameters, then the initial guess is just the first input camera pose
 	// Otherwise, a starting pose will need to be computed for each image
-	if ((EnumHasAnyFlags(SolverFlags, ECalibrationFlags::UseExtrinsicGuess)))
-	{
-		cv::Mat Rotation = Solver.Params.rowRange(NumIntrinsics, NumIntrinsics + 3);
-		cv::Mat Translation = Solver.Params.rowRange(NumIntrinsics + 3, NumIntrinsics + 6);
-
-		FOpenCVHelper::MakeObjectVectorsFromCameraPose(CameraPoses[0], Rotation, Translation);
-	}
-	else
+	for (int ImageIndex = 0; ImageIndex < NumPosesToSolve; ImageIndex++)
 	{
 		int ObjectPointsIndex = 0;
-		for (int ImageIndex = 0; ImageIndex < NumImages; ImageIndex++)
+
+		// Get a view to parameters used by the solver for the rotation and translation vectors for this image
+		const int ExtrinsicOffset = NumIntrinsics + (ImageIndex * NumExtrinsics);
+		cv::Mat Rotation = Solver.Params.rowRange(ExtrinsicOffset, ExtrinsicOffset + 3);
+		cv::Mat Translation = Solver.Params.rowRange(ExtrinsicOffset + 3, ExtrinsicOffset + 6);
+
+		if ((EnumHasAnyFlags(SolverFlags, ECalibrationFlags::UseExtrinsicGuess)))
+		{
+			FOpenCVHelper::MakeObjectVectorsFromCameraPose(CameraPoses[ImageIndex], Rotation, Translation);
+		}
+		else
 		{
 			const int NumImagePoints = NumPointsMat.at<int>(ImageIndex);
 
@@ -288,11 +311,6 @@ FDistortionCalibrationResult ULensDistortionSolverOpenCV::Solve_Implementation(
 
 			ObjectPointsIndex += NumImagePoints;
 
-			// Get a view to parameters used by the solver for the rotation and translation vectors for this image
-			const int ExtrinsicOffset = NumIntrinsics + (ImageIndex * NumExtrinsics);
-			cv::Mat Rotation = Solver.Params.rowRange(ExtrinsicOffset, ExtrinsicOffset + 3);
-			cv::Mat Translation = Solver.Params.rowRange(ExtrinsicOffset + 3, ExtrinsicOffset + 6);
-
 			InitCameraExtrinsics(LensModel, ObjectPointsInImage, ImagePointsInImage, CameraMatrix, DistCoeffs, CvImageSize, Rotation, Translation, SolverFlags);
 		}
 	}
@@ -301,7 +319,7 @@ FDistortionCalibrationResult ULensDistortionSolverOpenCV::Solve_Implementation(
 	// which represents how the camera moved between each image. The constrained solver, which only solves one camera pose, will use 
 	// this transformation to offset the pose used in the projection of points for each image. 
 	TArray<FTransform> CameraMovements;
-	if ((EnumHasAnyFlags(SolverFlags, ECalibrationFlags::UseExtrinsicGuess)))
+	if (bConstrainWithCameraMotion)
 	{
 		CameraMovements.Reserve(CameraPoses.Num());
 
@@ -324,6 +342,11 @@ FDistortionCalibrationResult ULensDistortionSolverOpenCV::Solve_Implementation(
 		bool bShouldProceed = Solver.UpdateAlt() && IsRunning();
 		bool bComputeJacobian = Solver.State == FLevMarqSolver::ESolverState::ComputeJacobian;
 
+		if (EnumHasAnyFlags(SolverFlags, ECalibrationFlags::FixAspectRatio))
+		{
+			SolverParams[0] = SolverParams[1] * PixelAspect;
+		}
+
 		// Update the camera matrix and distortion parameters with the latest values of the parameters from the solver
 		CameraMatrix.at<double>(0, 0) = SolverParams[0];
 		CameraMatrix.at<double>(1, 1) = SolverParams[1];
@@ -340,7 +363,7 @@ FDistortionCalibrationResult ULensDistortionSolverOpenCV::Solve_Implementation(
 
 		// If using a guess for the extrinsic parameters, cache the current solver pose as a FTransform to more easily offset the camera pose for each image
 		FTransform CurrentSolverPose;
-		if ((EnumHasAnyFlags(SolverFlags, ECalibrationFlags::UseExtrinsicGuess)))
+		if (bConstrainWithCameraMotion)
 		{
 			cv::Mat Rotation = Solver.Params.rowRange(NumIntrinsics, NumIntrinsics + 3);
 			cv::Mat Translation = Solver.Params.rowRange(NumIntrinsics + 3, NumIntrinsics + 6);
@@ -362,7 +385,7 @@ FDistortionCalibrationResult ULensDistortionSolverOpenCV::Solve_Implementation(
 			ObjectPointsIndex += NumImagePoints;
 
 			int ExtrinsicOffset = 0;
-			if ((EnumHasAnyFlags(SolverFlags, ECalibrationFlags::UseExtrinsicGuess)))
+			if (bConstrainWithCameraMotion)
 			{
 				ExtrinsicOffset = NumIntrinsics;
 			}
@@ -373,7 +396,7 @@ FDistortionCalibrationResult ULensDistortionSolverOpenCV::Solve_Implementation(
 
 			cv::Mat Rotation;
 			cv::Mat Translation;
-			if ((EnumHasAnyFlags(SolverFlags, ECalibrationFlags::UseExtrinsicGuess)))
+			if (bConstrainWithCameraMotion)
 			{
 				// Transform the solver's current camera pose by the camera movement to get the pose for this image
 				FTransform ImagePose = CurrentSolverPose * CameraMovements[ImageIndex];
@@ -440,6 +463,21 @@ FDistortionCalibrationResult ULensDistortionSolverOpenCV::Solve_Implementation(
 	Result.ImageCenter.PrincipalPoint.X = CameraMatrix.at<double>(0, 2);
 	Result.ImageCenter.PrincipalPoint.Y = CameraMatrix.at<double>(1, 2);
 
+	Result.CameraPoses.Reserve(NumImages);
+	for (int ImageIndex = 0; ImageIndex < NumImages; ++ImageIndex)
+	{
+		// Get the rotation and translation vectors for this image
+		int ExtrinsicOffset = NumIntrinsics + (ImageIndex * NumExtrinsics);
+		cv::Mat Rotation = Solver.Params.rowRange(ExtrinsicOffset, ExtrinsicOffset + 3);
+		cv::Mat Translation = Solver.Params.rowRange(ExtrinsicOffset + 3, ExtrinsicOffset + 6);
+
+		// Convert rotation and translation vectors to an FTransform in the UE coordinate space
+		FTransform CameraPoseForImage;
+		FOpenCVHelper::MakeCameraPoseFromObjectVectors(Rotation, Translation, CameraPoseForImage);
+
+		Result.CameraPoses.Add(CameraPoseForImage);
+	}
+
 	if (LensModel == UAnamorphicLensModel::StaticClass())
 	{
 		for (int CoeffIndex = 0; CoeffIndex < NumDistortionCoefficients; ++CoeffIndex)
@@ -1279,8 +1317,17 @@ void ULensDistortionSolverOpenCV::ProjectPointsSpherical(
 			{
 				double* JacFocalLengthPtr = JacFocalLength.ptr<double>(PointIndex * 2);
 
-				JacFocalLengthPtr[0] = DistortedPoint.x;
-				JacFocalLengthPtr[1] = 0;
+				if (EnumHasAnyFlags(SolverFlags, ECalibrationFlags::FixAspectRatio))
+				{
+					double PixelAspect = Fx / Fy;
+					JacFocalLengthPtr[0] = 0;
+					JacFocalLengthPtr[1] = DistortedPoint.x * PixelAspect;
+				}
+				else
+				{
+					JacFocalLengthPtr[0] = DistortedPoint.x;
+					JacFocalLengthPtr[1] = 0;
+				}
 
 				JacFocalLengthPtr[JacobianStep + 0] = 0;
 				JacFocalLengthPtr[JacobianStep + 1] = DistortedPoint.y;
diff --git a/Engine/Plugins/VirtualProduction/CameraCalibration/Source/CameraCalibrationEditor/Private/Calibrators/CameraCalibrationSolver.h b/Engine/Plugins/VirtualProduction/CameraCalibration/Source/CameraCalibrationEditor/Private/Calibrators/CameraCalibrationSolver.h
--- a/Engine/Plugins/VirtualProduction/CameraCalibration/Source/CameraCalibrationEditor/Private/Calibrators/CameraCalibrationSolver.h
+++ b/Engine/Plugins/VirtualProduction/CameraCalibration/Source/CameraCalibrationEditor/Private/Calibrators/CameraCalibrationSolver.h
@@ -56,7 +56,8 @@ enum class ECalibrationFlags : uint32
 	FixFocalLength = 1 << 2,     /** The solver will not optimize the focal length */
 	FixPrincipalPoint = 1 << 3,  /** The solver will not optimize the principal point */
 	FixExtrinsics = 1 << 4,      /** The solver will not optimize the camera extrinsics */
-	FixZeroDistortion = 1 << 5   /** The solver will fix all distortion values at 0 */
+	FixZeroDistortion = 1 << 5,  /** The solver will fix all distortion values at 0 */
+	FixAspectRatio = 1 << 6      /** The solver will respect the input aspect ratio when solving for Fx and Fy */
 };
 
 /** 
diff --git a/Engine/Plugins/VirtualProduction/CameraCalibration/Source/CameraCalibrationEditor/Private/Calibrators/CameraLensDistortionAlgoCheckerboard.cpp b/Engine/Plugins/VirtualProduction/CameraCalibration/Source/CameraCalibrationEditor/Private/Calibrators/CameraLensDistortionAlgoCheckerboard.cpp
--- a/Engine/Plugins/VirtualProduction/CameraCalibration/Source/CameraCalibrationEditor/Private/Calibrators/CameraLensDistortionAlgoCheckerboard.cpp
+++ b/Engine/Plugins/VirtualProduction/CameraCalibration/Source/CameraCalibrationEditor/Private/Calibrators/CameraLensDistortionAlgoCheckerboard.cpp
@@ -41,6 +41,7 @@ static TAutoConsoleVariable<bool> CVarUseIntrinsicsGuess(TEXT("LensDistortionChe
 static TAutoConsoleVariable<bool> CVarFixExtrinsics(TEXT("LensDistortionCheckerboard.FixExtrinsics"), false, TEXT("If true, the solver will fix the camera extrinsics to the user-provided camera poses"));
 static TAutoConsoleVariable<bool> CVarFixZeroDistortion(TEXT("LensDistortionCheckerboard.FixZeroDistortion"), false, TEXT("If true, the solver will fix all distortion values to always be 0"));
 static TAutoConsoleVariable<bool> CVarUseExtrinsicsGuess(TEXT("LensDistortionCheckerboard.UseExtrinsicsGuess"), false, TEXT("If true, the actual checkerboard and camera poses will be used when running the solver"));
+static TAutoConsoleVariable<bool> CVarUseTrackedCalibrator(TEXT("LensDistortionCheckerboard.UseTrackedCalibrator"), false, TEXT("If true, the 3D points of the checkerboard will come from the tracked pose of the object, otherwise, dummy points will be used."));
 #endif
 
 const int UCameraLensDistortionAlgoCheckerboard::DATASET_VERSION = 1;
@@ -350,35 +351,22 @@ bool UCameraLensDistortionAlgoCheckerboard::AddCalibrationRow(FText& OutErrorMes
  	}
 
 	// Fill out the checkerboard's 3D points
-	if (CVarUseExtrinsicsGuess.GetValueOnGameThread())
-	{
-		const FVector LocationTL = Calibrator->TopLeft->GetComponentLocation();
-		const FVector LocationTR = Calibrator->TopRight->GetComponentLocation();
-		const FVector LocationBL = Calibrator->BottomLeft->GetComponentLocation();
+	const FVector LocationTL = Calibrator->TopLeft->GetComponentLocation();
+	const FVector LocationTR = Calibrator->TopRight->GetComponentLocation();
+	const FVector LocationBL = Calibrator->BottomLeft->GetComponentLocation();
 
-		const FVector RightVector = LocationTR - LocationTL;
-		const FVector DownVector = LocationBL - LocationTL;
+	const FVector RightVector = LocationTR - LocationTL;
+	const FVector DownVector = LocationBL - LocationTL;
 
-		const float HorizontalStep = (Row->NumCornerCols > 1) ? (1.0f / (Row->NumCornerCols - 1)) : 0.0f;
-		const float VerticalStep = (Row->NumCornerRows > 1) ? (1.0f / (Row->NumCornerRows - 1)) : 0.0f;
+	const float HorizontalStep = (Row->NumCornerCols > 1) ? (1.0f / (Row->NumCornerCols - 1)) : 0.0f;
+	const float VerticalStep = (Row->NumCornerRows > 1) ? (1.0f / (Row->NumCornerRows - 1)) : 0.0f;
 
-		for (int32 RowIdx = 0; RowIdx < Row->NumCornerRows; ++RowIdx)
-		{
-			for (int32 ColIdx = 0; ColIdx < Row->NumCornerCols; ++ColIdx)
-			{
-				const FVector PointLocation = LocationTL + (RightVector * ColIdx * HorizontalStep) + (DownVector * RowIdx * VerticalStep);
-				Row->Points3d.Add(PointLocation);
-			}
-		}
-	}
-	else
+	for (int32 RowIdx = 0; RowIdx < Row->NumCornerRows; ++RowIdx)
 	{
-		for (int32 RowIdx = 0; RowIdx < Row->NumCornerRows; ++RowIdx)
+		for (int32 ColIdx = 0; ColIdx < Row->NumCornerCols; ++ColIdx)
 		{
-			for (int32 ColIdx = 0; ColIdx < Row->NumCornerCols; ++ColIdx)
-			{
-				Row->Points3d.Add(Row->SquareSideInCm * FVector(ColIdx, RowIdx, 0));
-			}
+			const FVector PointLocation = LocationTL + (RightVector * ColIdx * HorizontalStep) + (DownVector * RowIdx * VerticalStep);
+			Row->Points3d.Add(PointLocation);
 		}
 	}
 
@@ -677,7 +665,7 @@ FDistortionCalibrationTask UCameraLensDistortionAlgoCheckerboard::BeginCalibrati
 	}
 
 	const float FocalLengthEstimateValue = FocalLengthEstimate.GetValue();
-	const double Fx = (FocalLengthEstimateValue / DesqueezeSensorWidth) * ImageSize.X;
+	const double Fx = ((double)FocalLengthEstimateValue / DesqueezeSensorWidth) * ImageSize.X;
 
 	// When operating on a desqueezed image, we expect our pixel aspect to be square, so horizontal and vertical field of view are assumed to be equal (i.e. Fx == Fy)
 	FVector2D FocalLength = FVector2D(Fx, Fx);
@@ -708,6 +696,26 @@ FDistortionCalibrationTask UCameraLensDistortionAlgoCheckerboard::BeginCalibrati
 		CameraPoses.Add(Row->CameraData.Pose);
 	}
 
+	if (!CVarUseTrackedCalibrator.GetValueOnGameThread())
+	{
+		// Use a set of dummy points for the 3D checkerboard corners. The board is assumed to lie in the YZ plane with the TopLeft corner at (0, 0, 0) in world space.
+		// The side length of the checkerboard squares are used to position the remaining corners.
+		Samples3d.Empty();
+		for (int32 ImageIndex = 0; ImageIndex < CalibrationRows.Num(); ++ImageIndex)
+		{
+			FObjectPoints Points3d;
+			for (int32 RowIdx = 0; RowIdx < CalibrationRows[0]->NumCornerRows; ++RowIdx)
+			{
+				for (int32 ColIdx = 0; ColIdx < CalibrationRows[0]->NumCornerCols; ++ColIdx)
+				{
+					Points3d.Points.Add(CalibrationRows[0]->SquareSideInCm * FVector(0, ColIdx, -RowIdx));
+				}
+			}
+
+			Samples3d.Add(Points3d);
+		}
+	}
+
 	ECalibrationFlags SolverFlags = ECalibrationFlags::None;
 
 	if (CVarUseExtrinsicsGuess.GetValueOnGameThread())
@@ -744,7 +752,7 @@ FDistortionCalibrationTask UCameraLensDistortionAlgoCheckerboard::BeginCalibrati
 	
 	Solver = NewObject<ULensDistortionSolver>(this, Tool->GetSolverClass());
 
-	CalibrationTask = UE::Tasks::Launch(UE_SOURCE_LOCATION, [Solver=Solver, Model, Samples3d, Samples2d, ImageSize, FocalLength, ImageCenter, CameraPoses, PixelAspect, SolverFlags, Focus, Zoom]() mutable
+	CalibrationTask = UE::Tasks::Launch(UE_SOURCE_LOCATION, [Solver = Solver, Model, Samples3d, Samples2d, ImageSize, FocalLength, ImageCenter, CameraPoses, PixelAspect, SolverFlags, Focus, Zoom]() mutable
 		{
 			FDistortionCalibrationResult Result = Solver->Solve(
 				Samples3d,
diff --git a/Engine/Plugins/VirtualProduction/CameraCalibrationCore/Source/CameraCalibrationCore/Public/CameraCalibrationTypes.h b/Engine/Plugins/VirtualProduction/CameraCalibrationCore/Source/CameraCalibrationCore/Public/CameraCalibrationTypes.h
--- a/Engine/Plugins/VirtualProduction/CameraCalibrationCore/Source/CameraCalibrationCore/Public/CameraCalibrationTypes.h
+++ b/Engine/Plugins/VirtualProduction/CameraCalibrationCore/Source/CameraCalibrationCore/Public/CameraCalibrationTypes.h
@@ -34,6 +34,10 @@ struct FDistortionCalibrationResult
 	UPROPERTY(BlueprintReadWrite, Category = "Distortion")
 	FImageCenterInfo ImageCenter;
 
+	/** Calibrated camera pose for each input image */
+	UPROPERTY(BlueprintReadWrite, Category = "Distortion")
+	TArray<FTransform> CameraPoses;
+
 	/** Distortion parameters for the model specified by the lens file. And empty parameter array implies that there is a valid ST Map instead. */
 	UPROPERTY(BlueprintReadWrite, Category = "Distortion")
 	FDistortionInfo Parameters;
