CameraCalibration: Add support for an automated calibration test harness based on a json test description.
#rnx

diff --git a/Engine/Plugins/VirtualProduction/CameraCalibration/Source/CameraCalibrationEditor/Private/TestCameraCalibration.cpp b/Engine/Plugins/VirtualProduction/CameraCalibration/Source/CameraCalibrationEditor/Private/TestCameraCalibration.cpp
--- a/Engine/Plugins/VirtualProduction/CameraCalibration/Source/CameraCalibrationEditor/Private/TestCameraCalibration.cpp
+++ b/Engine/Plugins/VirtualProduction/CameraCalibration/Source/CameraCalibrationEditor/Private/TestCameraCalibration.cpp
@@ -1,30 +1,153 @@
 // Copyright Epic Games, Inc. All Rights Reserved.
 
+#include "TestCameraCalibration.h"
 
 #if WITH_DEV_AUTOMATION_TESTS
 
-#include "CoreMinimal.h"
-
 #include "Calibrators/CameraCalibrationSolver.h"
 #include "CameraCalibrationTypes.h"
 #include "CameraCalibrationUtilsPrivate.h"
 #include "Dialog/SCustomDialog.h"
-#include "Engine/Engine.h"
 #include "Engine/Texture2D.h"
-#include "LensData.h"
+#include "HAL/FileManager.h"
+#include "JsonObjectConverter.h"
+#include "Logging/LogMacros.h"
 #include "Misc/AutomationTest.h"
 #include "Models/SphericalLensModel.h"
 #include "OpenCVHelper.h"
-#include "Slate/DeferredCleanupSlateBrush.h"
 #include "TestCameraCalibrationSettings.h"
-#include "Widgets/Images/SImage.h"
+#include "UI/SImageTexture.h"
+
+DEFINE_LOG_CATEGORY_STATIC(LogTestCameraCalibration, Log, All);
 
 IMPLEMENT_SIMPLE_AUTOMATION_TEST(FTestDistortionSpherical, "Plugins.CameraCalibration.TestDistortionSpherical", EAutomationTestFlags::EditorContext | EAutomationTestFlags::EngineFilter)
 IMPLEMENT_SIMPLE_AUTOMATION_TEST(FTestNodalOffset, "Plugins.CameraCalibration.TestNodalOffset", EAutomationTestFlags::EditorContext | EAutomationTestFlags::EngineFilter)
 
 namespace UE::Private::CameraCalibration::AutomatedTests
 {
-	static TSharedPtr<FDeferredCleanupSlateBrush> TextureBrush;
+	namespace Logging
+	{
+		FString GetSolverFlagString(const TOptional<bool>& Flag)
+		{
+			if (Flag.IsSet() && Flag.GetValue())
+			{
+				return TEXT("true");
+			}
+			return TEXT("false");
+		}
+
+		void LogDistortionParameters(const TArray<float>& DistortionParameters)
+		{
+			if (ensure(DistortionParameters.Num() == 5))
+			{
+				UE_LOG(LogTestCameraCalibration, Log, TEXT("\t\t(%lf, %lf, %lf, %lf, %lf)"),
+					DistortionParameters[0], // K1
+					DistortionParameters[1], // K2
+					DistortionParameters[3], // P1
+					DistortionParameters[4], // P2
+					DistortionParameters[2]  // K3
+				);
+			}
+		}
+
+		void LogPoses(const TArray<FTransform>& Poses)
+		{
+			for (int32 PoseIndex = 0; PoseIndex < Poses.Num(); ++PoseIndex)
+			{
+				const FTransform& Pose = Poses[PoseIndex];
+				const FVector Location = Pose.GetTranslation();
+				const FRotator Rotation = Pose.GetRotation().Rotator();
+
+				UE_LOG(LogTestCameraCalibration, Log, TEXT("Pose %d:"), PoseIndex);
+				UE_LOG(LogTestCameraCalibration, Log, TEXT("\t\tLocation: (%lf, %lf, %lf)"), Location.X, Location.Y, Location.Z);
+				UE_LOG(LogTestCameraCalibration, Log, TEXT("\t\tRotation: (%lf, %lf, %lf)"), Rotation.Roll, Rotation.Pitch, Rotation.Yaw);
+			}
+		}
+
+		void LogPoses(const TArray<FLocationRotation>& Poses)
+		{
+			for (int32 PoseIndex = 0; PoseIndex < Poses.Num(); ++PoseIndex)
+			{
+				const FLocationRotation& Pose = Poses[PoseIndex];
+
+				UE_LOG(LogTestCameraCalibration, Log, TEXT("Pose %d:"), PoseIndex);
+				UE_LOG(LogTestCameraCalibration, Log, TEXT("\t\tLocation: (%lf, %lf, %lf)"), Pose.Location.X, Pose.Location.Y, Pose.Location.Z);
+				UE_LOG(LogTestCameraCalibration, Log, TEXT("\t\tRotation: (%lf, %lf, %lf)"), Pose.Rotation.Roll, Pose.Rotation.Pitch, Pose.Rotation.Yaw);
+			}
+		}
+
+		void LogTestDescription(const FCalibrationTest& Test)
+		{
+			UE_LOG(LogTestCameraCalibration, Log, TEXT("Test #%d:"), Test.TestIndex);
+
+			UE_LOG(LogTestCameraCalibration, Log, TEXT(" "));
+
+			if (Test.DatasetPath.IsEmpty())
+			{
+				UE_LOG(LogTestCameraCalibration, Log, TEXT("Camera Profile:"));
+				UE_LOG(LogTestCameraCalibration, Log, TEXT("Sensor Dimensions: (%.3lfmm x %.3lfmm)"), Test.CameraProfile.SensorSize.X, Test.CameraProfile.SensorSize.Y);
+				UE_LOG(LogTestCameraCalibration, Log, TEXT("Image Resolution: (%d, %d)"), Test.CameraProfile.ImageSize.X, Test.CameraProfile.ImageSize.Y);
+
+				const FVector2D FxFyInPixels = Test.CameraProfile.GetFxFyInPixels();
+				UE_LOG(LogTestCameraCalibration, Log, TEXT("Focal Length: %lf mm (%lf pixels)"), Test.CameraProfile.FocalLengthInMM, FxFyInPixels.X);
+
+				UE_LOG(LogTestCameraCalibration, Log, TEXT("Image Center: (%lf, %lf)"), Test.CameraProfile.ImageCenter.X, Test.CameraProfile.ImageCenter.Y);
+
+				UE_LOG(LogTestCameraCalibration, Log, TEXT("Distortion Parameters:"));
+				LogDistortionParameters(Test.CameraProfile.DistortionParameters);
+
+				UE_LOG(LogTestCameraCalibration, Log, TEXT(" "));
+
+				UE_LOG(LogTestCameraCalibration, Log, TEXT("Checkerboard Profile:"));
+				UE_LOG(LogTestCameraCalibration, Log, TEXT("Checkerboard Dimensions: %d columns by %d rows"), Test.CheckerboardProfile.CheckerboardDimensions.X, Test.CheckerboardProfile.CheckerboardDimensions.Y);
+				UE_LOG(LogTestCameraCalibration, Log, TEXT("Checkerboard Square Size: %.2lfcm"), Test.CheckerboardProfile.SquareSize);
+
+				UE_LOG(LogTestCameraCalibration, Log, TEXT(" "));
+
+				UE_LOG(LogTestCameraCalibration, Log, TEXT("Camera Poses:"));
+				LogPoses(Test.CameraPoses);
+
+				UE_LOG(LogTestCameraCalibration, Log, TEXT(" "));
+
+				UE_LOG(LogTestCameraCalibration, Log, TEXT("Checkerboard Poses:"));
+				LogPoses(Test.CheckerboardPoses);
+			}
+			else
+			{
+				UE_LOG(LogTestCameraCalibration, Log, TEXT("Using dataset: %s"), *Test.DatasetPath);
+			}
+
+			UE_LOG(LogTestCameraCalibration, Log, TEXT(" "));
+
+			UE_LOG(LogTestCameraCalibration, Log, TEXT("Solver Settings:"));
+			UE_LOG(LogTestCameraCalibration, Log, TEXT("Focal Length Guess: %lf mm"), Test.SolverSettings.FocalLengthGuess);
+			UE_LOG(LogTestCameraCalibration, Log, TEXT("Image Center Guess: (%lf, %lf)"), Test.SolverSettings.ImageCenterGuess.X, Test.SolverSettings.ImageCenterGuess.Y);
+			UE_LOG(LogTestCameraCalibration, Log, TEXT("Use Intrinsic Guess: %s"), *GetSolverFlagString(Test.SolverSettings.bUseIntrinsicGuess));
+			UE_LOG(LogTestCameraCalibration, Log, TEXT("Use Extrinsic Guess: %s"), *GetSolverFlagString(Test.SolverSettings.bUseExtrinsicGuess));
+			UE_LOG(LogTestCameraCalibration, Log, TEXT("Fix Focal Length: %s"), *GetSolverFlagString(Test.SolverSettings.bFixFocalLength));
+			UE_LOG(LogTestCameraCalibration, Log, TEXT("Fix Image Center: %s"), *GetSolverFlagString(Test.SolverSettings.bFixPrincipalPoint));
+			UE_LOG(LogTestCameraCalibration, Log, TEXT("Fix Extrinsics: %s"), *GetSolverFlagString(Test.SolverSettings.bFixExtrinsics));
+			UE_LOG(LogTestCameraCalibration, Log, TEXT("Fix Zero Distortion: %s"), *GetSolverFlagString(Test.SolverSettings.bFixZeroDistortion));
+			UE_LOG(LogTestCameraCalibration, Log, TEXT("Fix Aspect Ratio: %s"), *GetSolverFlagString(Test.SolverSettings.bFixAspectRatio));
+		}
+
+		void LogCalibrationResult(const FDistortionCalibrationResult& Result, double FocalLengthInMM)
+		{
+			UE_LOG(LogTestCameraCalibration, Log, TEXT(" "));
+			UE_LOG(LogTestCameraCalibration, Log, TEXT("Calibration Result:"));
+			UE_LOG(LogTestCameraCalibration, Log, TEXT(" "));
+			UE_LOG(LogTestCameraCalibration, Log, TEXT("FxFy: (%lf, %lf) pixels (%lf mm)"), Result.FocalLength.FxFy.X, Result.FocalLength.FxFy.Y, FocalLengthInMM);
+			UE_LOG(LogTestCameraCalibration, Log, TEXT("Image Center: (%lf, %lf)"), Result.ImageCenter.PrincipalPoint.X, Result.ImageCenter.PrincipalPoint.Y);
+
+			UE_LOG(LogTestCameraCalibration, Log, TEXT("Distortion Parameters:"));
+			LogDistortionParameters(Result.Parameters.Parameters);
+
+			UE_LOG(LogTestCameraCalibration, Log, TEXT("Camera Poses:"));
+			LogPoses(Result.CameraPoses);
+
+			UE_LOG(LogTestCameraCalibration, Log, TEXT("RMSE: %.9lf"), Result.ReprojectionError);
+		}
+	}
 
 	void TestDistortionCalibration(FAutomationTestBase& Test)
 	{
@@ -157,35 +280,6 @@ namespace UE::Private::CameraCalibration::AutomatedTests
 			EstimatedCameraPoses.Add(EstimatedCameraPoseForView);
 		}
 
-		// Debug setting that draws all of the projected 2D checkerboard patterns to an image that can be displayed to the user to validate that the points make sense (and are within the bounds of the image)
-		if (TestSettings->bShowCheckerboardImage)
- 		{
-			UTexture2D* DebugTexture = UTexture2D::CreateTransient(ImageSize.X, ImageSize.Y, EPixelFormat::PF_B8G8R8A8);
-			UE::CameraCalibration::Private::ClearTexture(DebugTexture, FColor::Black);
-
-			for (const TArray<FVector2f>& Image : ImagePoints)
-			{
-				FOpenCVHelper::DrawCheckerboardCorners(Image, CheckerboardCornerDimensions, DebugTexture);
-			}
-
-			TextureBrush = FDeferredCleanupSlateBrush::CreateBrush(DebugTexture);
-
-			TSharedRef<SCustomDialog> DetectionWindow =
-				SNew(SCustomDialog)
-				.UseScrollBox(false)
-				.Content()
-				[
-					SNew(SImage)
-					.Image(TextureBrush->GetSlateBrush())
-				]
-				.Buttons
-				({
-					SCustomDialog::FButton(NSLOCTEXT("TestCameraCalibration", "OkButton", "Ok")),
-				});
-
-			DetectionWindow->Show();
- 		}
-
 		/**
 		 * Step 2: Introduce errors into the 3D and 2D point data to simulate real-world inaccuracies that occur when doing calibration
 		 */
@@ -445,6 +539,325 @@ namespace UE::Private::CameraCalibration::AutomatedTests
 			Test.AddInfo(FString::Printf(TEXT("\t\t\t\tRotation:    (%lf, %lf, %lf)"), Rotator.Roll, Rotator.Pitch, Rotator.Yaw));
 		}
 	}
+
+	void CopyDefaultsFromBase(const UStruct* StructDef, void* TestStruct, void* BaseStruct, void* DefaultsStruct)
+	{
+		// Iterate over each of the struct properties
+		for (TFieldIterator<FProperty> PropIt(StructDef); PropIt; ++PropIt)
+		{
+			FProperty* Property = *PropIt;
+
+			// If the current property is another struct, recursively check its properties
+			if (Property->IsA<FStructProperty>())
+			{
+				FStructProperty* StructProperty = CastFieldChecked<FStructProperty>(Property);
+
+				void* TmpTestValue = Property->ContainerPtrToValuePtr<uint8>(TestStruct);
+				void* TmpBaseValue = Property->ContainerPtrToValuePtr<uint8>(BaseStruct);
+				void* TmpDefaultValue = Property->ContainerPtrToValuePtr<uint8>(DefaultsStruct);
+
+				CopyDefaultsFromBase(StructProperty->Struct, TmpTestValue, TmpBaseValue, TmpDefaultValue);
+			}
+			else
+			{
+				// If the value of the property in the current test struct is the default value of that property,
+				// copy the value of the property from the base struct into the current test struct.
+				void* TestValue = Property->ContainerPtrToValuePtr<uint8>(TestStruct);
+				void* DefaultValue = Property->ContainerPtrToValuePtr<uint8>(DefaultsStruct);
+				if (Property->Identical(TestValue, DefaultValue, PPF_None))
+				{
+					void* ValueToCopy = Property->ContainerPtrToValuePtr<uint8>(BaseStruct);
+					Property->SetValue_InContainer(TestStruct, ValueToCopy);
+				}
+			}
+		}
+	}
+
+	void LoadCalibrationTestsFromFile(const FString FileName, FCalibrationTestSet& TestSet)
+	{
+		if (TUniquePtr<FArchive> FileReader = TUniquePtr<FArchive>(IFileManager::Get().CreateFileReader(*FileName)))
+		{
+			TSharedRef< TJsonReader<UTF8CHAR> > JsonReader = TJsonReaderFactory<UTF8CHAR>::Create(FileReader.Get());
+
+			TSharedPtr<FJsonObject> JsonObject = MakeShared<FJsonObject>();
+			if (FJsonSerializer::Deserialize(JsonReader, JsonObject))
+			{
+				if (FJsonObjectConverter::JsonObjectToUStruct<FCalibrationTestSet>(JsonObject.ToSharedRef(), &TestSet))
+				{
+					FCalibrationTest DefaultCalibrationTest;
+					for (FCalibrationTest& TestStruct : TestSet.Tests)
+					{
+						// If any of the tests specify a base test index, all of its properties that were not explicitly written in the json file will be copied from that base test
+						const int32 BaseTestIndex = TestStruct.BaseTestIndex;
+						if (BaseTestIndex >= 0)
+						{
+							FCalibrationTest* BaseTestPtr = TestSet.Tests.FindByPredicate([BaseTestIndex](const FCalibrationTest& TestStruct) { return TestStruct.TestIndex == BaseTestIndex; });
+							CopyDefaultsFromBase(FCalibrationTest::StaticStruct(), &TestStruct, BaseTestPtr, &DefaultCalibrationTest);
+						}
+					}
+				}
+				else
+				{
+					UE_LOG(LogTestCameraCalibration, Error, TEXT("Failed to convert json object to Calibration Test Set structure."));
+				}
+			}
+			else
+			{
+				UE_LOG(LogTestCameraCalibration, Error, TEXT("Failed to deserialize json file. Check that it is properly formatted."));
+			}
+		}
+		else
+		{
+			UE_LOG(LogTestCameraCalibration, Error, TEXT("Failed to read test filename"));
+		}
+	}
+
+	void LoadDatasetFromFile(FCalibrationTest& CalibrationTest, TArray<FObjectPoints>& ObjectPoints, TArray<FImagePoints>& ImagePoints)
+	{
+		// If the dataset path is not absolute, assume that it is relative to the project content directory
+		if (FPaths::IsRelative(CalibrationTest.DatasetPath))
+		{
+			CalibrationTest.DatasetPath = FPaths::ProjectContentDir() / CalibrationTest.DatasetPath;
+		}
+
+		// Find all json files in the selected directory
+		TArray<FString> FoundFiles;
+		const FString FileExtension = TEXT(".json");
+		IFileManager::Get().FindFiles(FoundFiles, *CalibrationTest.DatasetPath, *FileExtension);
+
+		// Early-out if selected directory has no json files to import
+		if (FoundFiles.Num() < 1)
+		{
+			UE_LOG(LogTestCameraCalibration, Error, TEXT("The following dataset had no json files to import: %s"), *CalibrationTest.DatasetPath);
+			return;
+		}
+
+		for (const FString& File : FoundFiles)
+		{
+			const FString JsonFileName = CalibrationTest.DatasetPath / File;
+
+			// Open the Json file for reading, and initialize a JsonReader to parse the contents
+			if (TUniquePtr<FArchive> FileReader = TUniquePtr<FArchive>(IFileManager::Get().CreateFileReader(*JsonFileName)))
+			{
+				TSharedRef< TJsonReader<> > JsonReader = TJsonReaderFactory<>::Create(FileReader.Get());
+
+				// Deserialize the row data from the Json file into a Json object
+				TSharedPtr<FJsonObject> JsonData = MakeShared<FJsonObject>();
+				if (FJsonSerializer::Deserialize(JsonReader, JsonData))
+				{
+					FCalibrationDatasetImage LoadedDataset;
+					if (FJsonObjectConverter::JsonObjectToUStruct<FCalibrationDatasetImage>(JsonData.ToSharedRef(), &LoadedDataset))
+					{
+						// Copy the loaded 3D and 2D points from this json file into the output point sets
+						FObjectPoints Point3d;
+						Point3d.Points = LoadedDataset.Points3d;
+
+						FImagePoints Points2d;
+						Points2d.Points = LoadedDataset.Points2d;
+
+						ObjectPoints.Add(Point3d);
+						ImagePoints.Add(Points2d);
+
+						// Copy the image dimensions in the json file to the calibration test's camera profile
+						CalibrationTest.CameraProfile.ImageSize = FIntPoint(LoadedDataset.ImageWidth, LoadedDataset.ImageHeight);
+					}
+				}
+			}
+		}
+	}
+
+	void GenerateCalibratorPoints(const FCheckerboardProfile& CheckerboardProfile, const TArray<FTransform>& CheckerboardPoses, TArray<FObjectPoints>& OutCheckerboardPoints)
+	{
+		OutCheckerboardPoints.Empty();
+		OutCheckerboardPoints.Reserve(CheckerboardPoses.Num());
+
+		const FIntPoint CheckerboardCornerDimensions = CheckerboardProfile.GetCornerDimensions();
+		const int32 NumInnerCorners = CheckerboardCornerDimensions.X * CheckerboardCornerDimensions.Y;
+
+		// Compute the location of each checkerboard corner as if the board were centered at the origin and lying in the YZ plane
+		FObjectPoints CheckerboardPointsAtOrigin;
+
+		// The top left intersection is one square length to the right and down from the actual top left corner of the board
+		const double TopLeftCornerX = 0.0 - (CheckerboardProfile.CheckerboardDimensions.X * 0.5 * CheckerboardProfile.SquareSize) + CheckerboardProfile.SquareSize;
+		const double TopLeftCornerY = 0.0 + (CheckerboardProfile.CheckerboardDimensions.Y * 0.5 * CheckerboardProfile.SquareSize) - CheckerboardProfile.SquareSize;
+
+		for (int32 RowIndex = 0; RowIndex < CheckerboardCornerDimensions.Y; ++RowIndex)
+		{
+			for (int32 ColumnIndex = 0; ColumnIndex < CheckerboardCornerDimensions.X; ++ColumnIndex)
+			{
+				const double Y = TopLeftCornerX + (CheckerboardProfile.SquareSize * ColumnIndex);
+				const double Z = TopLeftCornerY - (CheckerboardProfile.SquareSize * RowIndex);
+				CheckerboardPointsAtOrigin.Points.Add(FVector(0, Y, Z));
+			}
+		}
+
+		// Generate a set of checkerboard points for each pose by transforming the points centered at the origin by each checkerboard pose
+		for (const FTransform& Pose : CheckerboardPoses)
+		{
+			FObjectPoints CheckerboardPoints;
+			CheckerboardPoints.Points.Reserve(NumInnerCorners);
+			for (const FVector& Point : CheckerboardPointsAtOrigin.Points)
+			{
+				CheckerboardPoints.Points.Add(Pose.TransformPosition(Point));
+			}
+
+			OutCheckerboardPoints.Add(CheckerboardPoints);
+		}
+	}
+
+	bool ProjectCalibratorPoints(const FCameraProfile& CameraProfile, const TArray<FTransform>& CameraPoses, const TArray<FObjectPoints>& CalibratorPoints, TArray<FImagePoints>& OutImagePoints)
+	{
+		const int32 NumTestImages = CalibratorPoints.Num();
+		if (NumTestImages != CameraPoses.Num())
+		{
+			return false;
+		}
+
+		OutImagePoints.Empty();
+		OutImagePoints.Reserve(NumTestImages);
+
+		for (int32 ImageIndex = 0; ImageIndex < NumTestImages; ++ImageIndex)
+		{
+			const FObjectPoints& CalibratorPointsForImage = CalibratorPoints[ImageIndex];
+			const FTransform& CameraPoseForImage = CameraPoses[ImageIndex];
+
+			FImagePoints Image;
+			bool bResult = FOpenCVHelper::ProjectPoints(CalibratorPointsForImage.Points, CameraProfile.GetFxFyInPixels(), CameraProfile.ImageCenter, CameraProfile.DistortionParameters, CameraPoseForImage, Image.Points);
+			if (!bResult)
+			{
+				return false;
+			}
+
+			OutImagePoints.Add(Image);
+		}
+
+		return true;
+	}
+
+	void DrawDebugCoverage(const TArray<FImagePoints>& CheckerboardImages, FIntPoint CheckerboardCornerDimensions, FIntPoint ImageSize)
+	{
+		const UTestCameraCalibrationSettings* TestSettings = GetDefault<UTestCameraCalibrationSettings>();
+		if (TestSettings->bShowCheckerboardImage)
+		{
+			UTexture2D* DebugTexture = UTexture2D::CreateTransient(ImageSize.X, ImageSize.Y, EPixelFormat::PF_B8G8R8A8);
+			UE::CameraCalibration::Private::ClearTexture(DebugTexture, FColor::Black);
+
+			for (const FImagePoints& Image : CheckerboardImages)
+			{
+				FOpenCVHelper::DrawCheckerboardCorners(Image.Points, CheckerboardCornerDimensions, DebugTexture);
+			}
+
+			TSharedRef<SCustomDialog> DebugImageDialog =
+				SNew(SCustomDialog)
+				.UseScrollBox(false)
+				.Content()
+				[
+					SNew(SImageTexture, DebugTexture)
+				]
+			.Buttons
+			({
+				SCustomDialog::FButton(NSLOCTEXT("TestCameraCalibration", "OkButton", "Ok")),
+				});
+
+			DebugImageDialog->Show();
+		}
+	}
+
+	ECalibrationFlags GetCalibrationFlags(const FSolverSettings& SolverSettings)
+	{
+		static auto TestOptional = [](const TOptional<bool>& Value) {return Value && *Value; };
+
+		ECalibrationFlags SolverFlags = ECalibrationFlags::None;
+		if (TestOptional(SolverSettings.bUseIntrinsicGuess))
+		{
+			EnumAddFlags(SolverFlags, ECalibrationFlags::UseIntrinsicGuess);
+		}
+		if (TestOptional(SolverSettings.bUseExtrinsicGuess))
+		{
+			EnumAddFlags(SolverFlags, ECalibrationFlags::UseExtrinsicGuess);
+		}
+		if (TestOptional(SolverSettings.bFixFocalLength))
+		{
+			EnumAddFlags(SolverFlags, ECalibrationFlags::FixFocalLength);
+		}
+		if (TestOptional(SolverSettings.bFixPrincipalPoint))
+		{
+			EnumAddFlags(SolverFlags, ECalibrationFlags::FixPrincipalPoint);
+		}
+		if (TestOptional(SolverSettings.bFixExtrinsics))
+		{
+			EnumAddFlags(SolverFlags, ECalibrationFlags::FixExtrinsics);
+		}
+		if (TestOptional(SolverSettings.bFixZeroDistortion))
+		{
+			EnumAddFlags(SolverFlags, ECalibrationFlags::FixZeroDistortion);
+		}
+		if (TestOptional(SolverSettings.bFixAspectRatio))
+		{
+			EnumAddFlags(SolverFlags, ECalibrationFlags::FixAspectRatio);
+		}
+		return SolverFlags;
+	}
+
+	void TestCameraCalibration(FCalibrationTest& CalibrationTest, bool bLogVerboseTestDescription)
+	{
+		if (bLogVerboseTestDescription)
+		{
+			Logging::LogTestDescription(CalibrationTest);
+		}
+
+		TArray<FObjectPoints> CalibratorPoints3D;
+		TArray<FImagePoints> ImagePoints2D;
+
+		// If no dataset path is provided, the 3D and 2D points will be generated from the test description
+		if (CalibrationTest.DatasetPath.IsEmpty())
+		{
+			// Generate the 3D checkerboard points based on the checkerboard description and poses
+			TArray<FTransform> CheckerboardPoses;
+			CalibrationTest.GetCheckerboardTransforms(CheckerboardPoses);
+
+			GenerateCalibratorPoints(CalibrationTest.CheckerboardProfile, CheckerboardPoses, CalibratorPoints3D);
+
+			// Generate the 2D image points by projecting the 3D checkerboard points using the camera properties and poses
+			TArray<FTransform> CameraPoses;
+			CalibrationTest.GetCameraTransforms(CameraPoses);
+
+			ProjectCalibratorPoints(CalibrationTest.CameraProfile, CameraPoses, CalibratorPoints3D, ImagePoints2D);
+
+			// Pop up a dialog window with an image showing a debug view of the calibration patterns for this set of images
+			DrawDebugCoverage(ImagePoints2D, CalibrationTest.CheckerboardProfile.GetCornerDimensions(), CalibrationTest.CameraProfile.ImageSize);
+		}
+		else
+		{
+			LoadDatasetFromFile(CalibrationTest, CalibratorPoints3D, ImagePoints2D);
+		}
+
+		ULensDistortionSolverOpenCV* TestSolver = NewObject<ULensDistortionSolverOpenCV>();
+
+		const FSolverSettings SolverSettings = CalibrationTest.SolverSettings;
+		const ECalibrationFlags SolverFlags = GetCalibrationFlags(SolverSettings);
+
+		// TODO: Allow for camera pose guesses in test description
+		TArray<FTransform> CameraPoseGuesses;
+		CameraPoseGuesses.Empty();
+
+		const FVector2D FocalLengthGuessInPixels = CalibrationTest.CameraProfile.ConvertFocalLengthToPixels(SolverSettings.FocalLengthGuess);
+
+		FDistortionCalibrationResult Result = TestSolver->Solve(
+			CalibratorPoints3D,
+			ImagePoints2D,
+			CalibrationTest.CameraProfile.ImageSize,
+			FocalLengthGuessInPixels,
+			SolverSettings.ImageCenterGuess,
+			CameraPoseGuesses,
+			USphericalLensModel::StaticClass(),
+			1.0,
+			SolverFlags
+		);
+
+		const double FocalLengthInMM = (Result.FocalLength.FxFy.X / CalibrationTest.CameraProfile.ImageSize.X) * CalibrationTest.CameraProfile.SensorSize.X;
+		Logging::LogCalibrationResult(Result, FocalLengthInMM);
+	}
 }
 
 bool FTestDistortionSpherical::RunTest(const FString& Parameters)
@@ -459,4 +872,20 @@ bool FTestNodalOffset::RunTest(const FString& Parameters)
 	return true;
 }
 
+void UAutoCalibrationTest::RunTests(FString Filename, bool bLogVerboseTestDescription)
+{
+	if (FPaths::IsRelative(Filename))
+	{
+		Filename = FPaths::ProjectContentDir() / Filename;
+	}
+
+	FCalibrationTestSet CalibrationTestSet;
+	UE::Private::CameraCalibration::AutomatedTests::LoadCalibrationTestsFromFile(Filename, CalibrationTestSet);
+
+	for (FCalibrationTest& Test : CalibrationTestSet.Tests)
+	{
+		UE::Private::CameraCalibration::AutomatedTests::TestCameraCalibration(Test, bLogVerboseTestDescription);
+	}
+}
+
 #endif // WITH_DEV_AUTOMATION_TESTS
diff --git a/Engine/Plugins/VirtualProduction/CameraCalibration/Source/CameraCalibrationEditor/Private/TestCameraCalibration.h b/Engine/Plugins/VirtualProduction/CameraCalibration/Source/CameraCalibrationEditor/Private/TestCameraCalibration.h
new file mode 100644
--- /dev/null
+++ b/Engine/Plugins/VirtualProduction/CameraCalibration/Source/CameraCalibrationEditor/Private/TestCameraCalibration.h
@@ -0,0 +1,197 @@
+// Copyright Epic Games, Inc. All Rights Reserved.
+
+#pragma once
+
+#include "TestCameraCalibration.generated.h"
+
+/** Exposes RunTests to BP and Python for ease of running a test in the editor */
+UCLASS()
+class UAutoCalibrationTest : public UObject
+{
+	GENERATED_BODY()
+
+public:
+	/** Runs the calibration test set defined in the input file */
+	UFUNCTION(BlueprintCallable, Category = "Calibration")
+	static void RunTests(FString Filename, bool bLogVerboseTestDescription = false);
+};
+
+/** Useful to deserialize a single json file from a calibration dataset */
+USTRUCT()
+struct FCalibrationDatasetImage
+{
+	GENERATED_BODY()
+
+	UPROPERTY()
+	TArray<FVector2D> Points2d;
+
+	UPROPERTY()
+	TArray<FVector> Points3d;
+
+	UPROPERTY()
+	int32 ImageWidth = -1;
+
+	UPROPERTY()
+	int32 ImageHeight = -1;
+};
+
+/**
+ * Useful because this struct will be more human-readable when serialized to json than FTransform.
+ * FRotator is easier to read/write than FTransform's FQuat, and the automated tests do not use FTransform's scale.
+ */
+USTRUCT()
+struct FLocationRotation
+{
+	GENERATED_BODY()
+
+	UPROPERTY()
+	FVector Location;
+
+	UPROPERTY()
+	FRotator Rotation;
+};
+
+/** Description of a camera / lens, useful for projecting 3D calibrator points to 2D */
+USTRUCT()
+struct FCameraProfile
+{
+	GENERATED_BODY()
+
+	UPROPERTY()
+	FVector2D SensorSize = FVector2D(-1.0);
+
+	UPROPERTY()
+	FIntPoint ImageSize = FIntPoint(-1);
+
+	UPROPERTY()
+	double FocalLengthInMM = -1.0;
+
+	UPROPERTY()
+	FVector2D ImageCenter = FVector2D(-1.0);
+
+	UPROPERTY()
+	TArray<float> DistortionParameters;
+
+	FVector2D ConvertFocalLengthToPixels(double FocalLength) const
+	{
+		return (FVector2D(FocalLength) / SensorSize) * FVector2D(ImageSize);
+	}
+
+	FVector2D GetFxFyInPixels() const
+	{
+		return ConvertFocalLengthToPixels(FocalLengthInMM);
+	}
+};
+
+/** Description of a checkerboard calibrator, useful for generating sets of 3D calibration points in world space */
+USTRUCT()
+struct FCheckerboardProfile
+{
+	GENERATED_BODY()
+
+	UPROPERTY()
+	FIntPoint CheckerboardDimensions = FIntPoint(-1);
+
+	UPROPERTY()
+	double SquareSize = -1.0;
+
+	FIntPoint GetCornerDimensions() const
+	{
+		return CheckerboardDimensions - FIntPoint(1, 1);
+	}
+};
+
+/** Settings for the camera calibration solver */
+USTRUCT()
+struct FSolverSettings
+{
+	GENERATED_BODY()
+
+	UPROPERTY()
+	double FocalLengthGuess = -1.0;
+
+	UPROPERTY()
+	FVector2D ImageCenterGuess = FVector2D(-1.0);
+
+	UPROPERTY()
+	TOptional<bool> bUseIntrinsicGuess;
+	
+	UPROPERTY()
+	TOptional<bool> bUseExtrinsicGuess;
+	
+	UPROPERTY()
+	TOptional<bool> bFixFocalLength;
+	
+	UPROPERTY()
+	TOptional<bool> bFixPrincipalPoint;
+	
+	UPROPERTY()
+	TOptional<bool> bFixExtrinsics;
+	
+	UPROPERTY()
+	TOptional<bool> bFixZeroDistortion;
+	
+	UPROPERTY()
+	TOptional<bool> bFixAspectRatio;
+};
+
+/** Full camera calibration test description */
+USTRUCT()
+struct FCalibrationTest
+{
+	GENERATED_BODY()
+
+	UPROPERTY()
+	int32 TestIndex = 0;
+
+	UPROPERTY()
+	int32 BaseTestIndex = -1;
+
+	UPROPERTY()
+	FString DatasetPath;
+
+	UPROPERTY()
+	FCameraProfile CameraProfile;
+	
+	UPROPERTY()
+	FCheckerboardProfile CheckerboardProfile;
+
+	UPROPERTY()
+	TArray<FLocationRotation> CameraPoses;
+	
+	UPROPERTY()
+	TArray<FLocationRotation> CheckerboardPoses;
+
+	UPROPERTY()
+	FSolverSettings SolverSettings;
+
+	void GetCameraTransforms(TArray<FTransform>& Transforms) const
+	{
+		GetTransforms(CameraPoses, Transforms);
+	}
+
+	void GetCheckerboardTransforms(TArray<FTransform>& Transforms) const
+	{
+		GetTransforms(CheckerboardPoses, Transforms);
+	}
+
+	void GetTransforms(const TArray<FLocationRotation>& Poses, TArray<FTransform>& Transforms) const
+	{
+		Transforms.Empty();
+		Transforms.Reserve(Poses.Num());
+		for (const FLocationRotation& Pose : Poses)
+		{
+			Transforms.Add(FTransform(Pose.Rotation, Pose.Location));
+		}
+	}
+};
+
+/** Set of camera calibration test descriptions */
+USTRUCT()
+struct FCalibrationTestSet
+{
+	GENERATED_BODY()
+
+	UPROPERTY()
+	TArray<FCalibrationTest> Tests;
+};
