CameraCalibration: Fix regression in solver that caused only one camera pose to be initialized.
#rnx

diff --git a/Engine/Plugins/VirtualProduction/CameraCalibration/Source/CameraCalibrationEditor/Private/Calibrators/CameraCalibrationSolver.cpp b/Engine/Plugins/VirtualProduction/CameraCalibration/Source/CameraCalibrationEditor/Private/Calibrators/CameraCalibrationSolver.cpp
--- a/Engine/Plugins/VirtualProduction/CameraCalibration/Source/CameraCalibrationEditor/Private/Calibrators/CameraCalibrationSolver.cpp
+++ b/Engine/Plugins/VirtualProduction/CameraCalibration/Source/CameraCalibrationEditor/Private/Calibrators/CameraCalibrationSolver.cpp
@@ -288,10 +288,9 @@ FDistortionCalibrationResult ULensDistortionSolverOpenCV::Solve_Implementation(
 	// Initialize the starting guess for the camera's extrinsic parameters. 
 	// If using a guess for the extrinsic parameters, then the initial guess is just the first input camera pose
 	// Otherwise, a starting pose will need to be computed for each image
+	int ObjectPointsIndex = 0;
 	for (int ImageIndex = 0; ImageIndex < NumPosesToSolve; ImageIndex++)
 	{
-		int ObjectPointsIndex = 0;
-
 		// Get a view to parameters used by the solver for the rotation and translation vectors for this image
 		const int ExtrinsicOffset = NumIntrinsics + (ImageIndex * NumExtrinsics);
 		cv::Mat Rotation = Solver.Params.rowRange(ExtrinsicOffset, ExtrinsicOffset + 3);
@@ -373,7 +372,7 @@ FDistortionCalibrationResult ULensDistortionSolverOpenCV::Solve_Implementation(
 
 		ReprojectionError = 0;
 
-		int ObjectPointsIndex = 0;
+		ObjectPointsIndex = 0;
 		for (int ImageIndex = 0; ImageIndex < NumImages; ImageIndex++)
 		{
 			int NumImagePoints = NumPointsMat.at<int>(ImageIndex);
