CameraCalibration: Added alternate curve mode to the curve editor that displays lens data tables with a focus axis instead of a zoom axis, and allows the values of those points to be edited. The curve is forced to be a linear interpolation with no tangent editing since there isn't any tangent data being stored in the data tables for the focus axis.

diff --git a/Engine/Plugins/VirtualProduction/CameraCalibration/Source/CameraCalibrationEditor/Private/AssetEditor/Curves/LensDataCurveModel.cpp b/Engine/Plugins/VirtualProduction/CameraCalibration/Source/CameraCalibrationEditor/Private/AssetEditor/Curves/LensDataCurveModel.cpp
--- a/Engine/Plugins/VirtualProduction/CameraCalibration/Source/CameraCalibrationEditor/Private/AssetEditor/Curves/LensDataCurveModel.cpp
+++ b/Engine/Plugins/VirtualProduction/CameraCalibration/Source/CameraCalibrationEditor/Private/AssetEditor/Curves/LensDataCurveModel.cpp
@@ -11,6 +11,7 @@ ECurveEditorViewID FLensDataCurveModel::ViewId = ECurveEditorViewID::Invalid;
 FLensDataCurveModel::FLensDataCurveModel(ULensFile* InOwner)
 	: FRichCurveEditorModel(InOwner)
 	, LensFile(InOwner)
+	, ClampOutputRange(TRange<double>(TNumericLimits<double>::Lowest(), TNumericLimits<double>::Max()))
 {
 	check(InOwner);
 
@@ -27,9 +28,38 @@ void FLensDataCurveModel::RemoveKeys(TArrayView<const FKeyHandle> InKeys)
 	//Don't support removing keys from curve editor by default. Specific models can override
 }
 
-void FLensDataCurveModel::SetKeyAttributes(TArrayView<const FKeyHandle> InKeys,	TArrayView<const FKeyAttributes> InAttributes, EPropertyChangeType::Type ChangeType)
+void FLensDataCurveModel::SetKeyPositions(TArrayView<const FKeyHandle> InKeys, TArrayView<const FKeyPosition> InKeyPositions, EPropertyChangeType::Type ChangeType)
 {
-	//Don't support changing interp attributes from curve editor by default. Specific models can override
+	// Re-implementation of FRichCurveEditorModel::SetKeyPosition that only changes the y-axis of the point being edited, as changing the x-axis (focus/zoom) is not supported
+	if (IsReadOnly())
+	{
+		return;
+	}
+
+	if (UObject* Owner = GetOwningObject())
+	{
+		if (IsValid())
+		{
+			Owner->Modify();
+
+			FRichCurve& RichCurve = GetRichCurve();
+			for (int32 Index = 0; Index < InKeys.Num(); ++Index)
+			{
+				FKeyHandle Handle = InKeys[Index];
+				if (RichCurve.IsKeyHandleValid(Handle))
+				{
+					const TRange<double> OutputRange = ClampOutputRange.Get();
+					RichCurve.GetKey(Handle).Value = FMath::Clamp(InKeyPositions[Index].OutputValue, OutputRange.GetLowerBoundValue(), OutputRange.GetUpperBoundValue());
+				}
+			}
+			
+			RichCurve.AutoSetTangents();
+			FPropertyChangedEvent PropertyChangeStruct(nullptr, ChangeType);
+			Owner->PostEditChangeProperty(PropertyChangeStruct);
+
+			CurveModifiedDelegate.Broadcast();
+		}
+	}
 }
 
 bool FLensDataCurveModel::IsValid() const
diff --git a/Engine/Plugins/VirtualProduction/CameraCalibration/Source/CameraCalibrationEditor/Private/AssetEditor/Curves/LensDataCurveModel.h b/Engine/Plugins/VirtualProduction/CameraCalibration/Source/CameraCalibrationEditor/Private/AssetEditor/Curves/LensDataCurveModel.h
--- a/Engine/Plugins/VirtualProduction/CameraCalibration/Source/CameraCalibrationEditor/Private/AssetEditor/Curves/LensDataCurveModel.h
+++ b/Engine/Plugins/VirtualProduction/CameraCalibration/Source/CameraCalibrationEditor/Private/AssetEditor/Curves/LensDataCurveModel.h
@@ -22,7 +22,7 @@ public:
 	//~ Begin FRichCurveEditorModel
 	virtual void AddKeys(TArrayView<const FKeyPosition> InKeyPositions, TArrayView<const FKeyAttributes> InAttributes, TArrayView<TOptional<FKeyHandle>>* OutKeyHandles) override;
 	virtual void RemoveKeys(TArrayView<const FKeyHandle> InKeys) override;
-	virtual void SetKeyAttributes(TArrayView<const FKeyHandle> InKeys, TArrayView<const FKeyAttributes> InAttributes, EPropertyChangeType::Type ChangeType = EPropertyChangeType::Unspecified) override;
+	virtual void SetKeyPositions(TArrayView<const FKeyHandle> InKeys, TArrayView<const FKeyPosition> InKeyPositions, EPropertyChangeType::Type ChangeType) override;
 	virtual bool IsValid() const override;
 	virtual FRichCurve& GetRichCurve() override;
 	virtual const FRichCurve& GetReadOnlyRichCurve() const override;
@@ -49,6 +49,9 @@ protected:
 	/** Active curve pointer */
 	FRichCurve CurrentCurve;
 
-	/** Wheter a valid curve was built from lens data */
+	/** An optional clamp on the output values (y-axis) that the curve keys are allowed to have */
+	TAttribute<TRange<double>> ClampOutputRange;
+	
+	/** Whether a valid curve was built from lens data */
 	bool bIsCurveValid = false;
 };
diff --git a/Engine/Plugins/VirtualProduction/CameraCalibration/Source/CameraCalibrationEditor/Private/AssetEditor/Curves/LensDistortionParametersCurveModel.cpp b/Engine/Plugins/VirtualProduction/CameraCalibration/Source/CameraCalibrationEditor/Private/AssetEditor/Curves/LensDistortionParametersCurveModel.cpp
deleted file mode 100644
diff --git a/Engine/Plugins/VirtualProduction/CameraCalibration/Source/CameraCalibrationEditor/Private/AssetEditor/Curves/LensDistortionParametersCurveModel.h b/Engine/Plugins/VirtualProduction/CameraCalibration/Source/CameraCalibrationEditor/Private/AssetEditor/Curves/LensDistortionParametersCurveModel.h
deleted file mode 100644
diff --git a/Engine/Plugins/VirtualProduction/CameraCalibration/Source/CameraCalibrationEditor/Private/AssetEditor/Curves/LensFocalLengthCurveModel.cpp b/Engine/Plugins/VirtualProduction/CameraCalibration/Source/CameraCalibrationEditor/Private/AssetEditor/Curves/LensFocalLengthCurveModel.cpp
deleted file mode 100644
diff --git a/Engine/Plugins/VirtualProduction/CameraCalibration/Source/CameraCalibrationEditor/Private/AssetEditor/Curves/LensFocalLengthCurveModel.h b/Engine/Plugins/VirtualProduction/CameraCalibration/Source/CameraCalibrationEditor/Private/AssetEditor/Curves/LensFocalLengthCurveModel.h
deleted file mode 100644
diff --git a/Engine/Plugins/VirtualProduction/CameraCalibration/Source/CameraCalibrationEditor/Private/AssetEditor/Curves/LensImageCenterCurveModel.cpp b/Engine/Plugins/VirtualProduction/CameraCalibration/Source/CameraCalibrationEditor/Private/AssetEditor/Curves/LensImageCenterCurveModel.cpp
deleted file mode 100644
diff --git a/Engine/Plugins/VirtualProduction/CameraCalibration/Source/CameraCalibrationEditor/Private/AssetEditor/Curves/LensImageCenterCurveModel.h b/Engine/Plugins/VirtualProduction/CameraCalibration/Source/CameraCalibrationEditor/Private/AssetEditor/Curves/LensImageCenterCurveModel.h
deleted file mode 100644
diff --git a/Engine/Plugins/VirtualProduction/CameraCalibration/Source/CameraCalibrationEditor/Private/AssetEditor/Curves/LensMultiAxisCurveModel.cpp b/Engine/Plugins/VirtualProduction/CameraCalibration/Source/CameraCalibrationEditor/Private/AssetEditor/Curves/LensMultiAxisCurveModel.cpp
new file mode 100644
--- /dev/null
+++ b/Engine/Plugins/VirtualProduction/CameraCalibration/Source/CameraCalibrationEditor/Private/AssetEditor/Curves/LensMultiAxisCurveModel.cpp
@@ -0,0 +1,104 @@
+﻿// Copyright Epic Games, Inc. All Rights Reserved.
+
+#include "LensMultiAxisCurveModel.h"
+
+#define LOCTEXT_NAMESPACE "FLensDataMultiAxisCurveModel"
+
+FLensDataMultiAxisCurveModel::FLensDataMultiAxisCurveModel(ULensFile* InOwner, ELensDataCategory InCategory, ELensCurveAxis InCurveAxis, float InCurveValue, int32 InParameterIndex)
+	: FLensDataCurveModel(InOwner)
+	, Category(InCategory)
+	, CurveAxis(InCurveAxis)
+	, CurveValue(InCurveValue)
+	, ParameterIndex(InParameterIndex)
+{
+	if (const FBaseLensTable* DataTable = LensFile->GetDataTable(Category))
+	{
+		if (CurveAxis == ELensCurveAxis::Zoom)
+		{
+			bIsCurveValid = DataTable->BuildParameterCurveAtFocus(InCurveValue, InParameterIndex, CurrentCurve);
+		}
+		else if (CurveAxis == ELensCurveAxis::Focus)
+		{
+			bIsCurveValid = DataTable->BuildParameterCurveAtZoom(InCurveValue, InParameterIndex, CurrentCurve);
+		}
+
+		ClampOutputRange = DataTable->GetCurveKeyPositionRange(InParameterIndex);
+	}
+}
+
+void FLensDataMultiAxisCurveModel::SetKeyPositions(TArrayView<const FKeyHandle> InKeys, TArrayView<const FKeyPosition> InKeyPositions, EPropertyChangeType::Type ChangeType)
+{
+	if (FBaseLensTable* DataTable = LensFile->GetDataTable(Category))
+	{
+		if (!DataTable->CanEditCurveKeyPositions(ParameterIndex))
+		{
+			return;
+		}
+
+		FLensDataCurveModel::SetKeyPositions(InKeys, InKeyPositions, ChangeType);
+		
+		if (CurveAxis == ELensCurveAxis::Zoom)
+		{
+			DataTable->SetParameterCurveKeysAtFocus(CurveValue, ParameterIndex, CurrentCurve, InKeys);
+		}
+		else if (CurveAxis == ELensCurveAxis::Focus)
+		{
+			DataTable->SetParameterCurveKeysAtZoom(CurveValue, ParameterIndex, CurrentCurve, InKeys);
+		}
+	}
+}
+
+void FLensDataMultiAxisCurveModel::SetKeyAttributes(TArrayView<const FKeyHandle> InKeys, TArrayView<const FKeyAttributes> InAttributes, EPropertyChangeType::Type ChangeType)
+{
+	// For now, don't allow focus axis curves to be editable
+	if (CurveAxis == ELensCurveAxis::Focus)
+	{
+		return;
+	}
+		
+	if (FBaseLensTable* DataTable = LensFile->GetDataTable(Category))
+	{
+		if (!DataTable->CanEditCurveKeyAttributes(ParameterIndex))
+		{
+			return;
+		}
+			
+		FLensDataCurveModel::SetKeyAttributes(InKeys, InAttributes, ChangeType);
+		
+		DataTable->SetParameterCurveKeysAtFocus(CurveValue, ParameterIndex, CurrentCurve, InKeys);
+	}
+}
+
+FText FLensDataMultiAxisCurveModel::GetKeyLabel() const
+{
+	if (CurveAxis == ELensCurveAxis::Zoom)
+	{
+		return LOCTEXT("ZoomAxisLabel", "Raw Zoom");
+	}
+	else
+	{
+		return LOCTEXT("FocusAxisLabel", "Raw Focus");
+	}
+}
+
+FText FLensDataMultiAxisCurveModel::GetValueLabel() const
+{
+	if (const FBaseLensTable* DataTable = LensFile->GetDataTable(Category))
+	{
+		return DataTable->GetParameterValueLabel(ParameterIndex);
+	}
+
+	return FText();
+}
+
+FText FLensDataMultiAxisCurveModel::GetValueUnitSuffixLabel() const
+{
+	if (const FBaseLensTable* DataTable = LensFile->GetDataTable(Category))
+	{
+		return DataTable->GetParameterValueUnitLabel(ParameterIndex);
+	}
+
+	return FText();
+}
+
+#undef LOCTEXT_NAMESPACE
\ No newline at end of file
diff --git a/Engine/Plugins/VirtualProduction/CameraCalibration/Source/CameraCalibrationEditor/Private/AssetEditor/Curves/LensMultiAxisCurveModel.h b/Engine/Plugins/VirtualProduction/CameraCalibration/Source/CameraCalibrationEditor/Private/AssetEditor/Curves/LensMultiAxisCurveModel.h
new file mode 100644
--- /dev/null
+++ b/Engine/Plugins/VirtualProduction/CameraCalibration/Source/CameraCalibrationEditor/Private/AssetEditor/Curves/LensMultiAxisCurveModel.h
@@ -0,0 +1,42 @@
+﻿// Copyright Epic Games, Inc. All Rights Reserved.
+
+#pragma once
+
+#include "LensDataCurveModel.h"
+#include "LensFile.h"
+
+enum class ELensCurveAxis : uint8
+{
+	Zoom,
+	Focus
+};
+
+/**
+ * Curve model for lens data tables that have multiple axes (e.g focus and zoom), which can generate curves for either axis
+ */
+class FLensDataMultiAxisCurveModel : public FLensDataCurveModel
+{
+public:
+	FLensDataMultiAxisCurveModel(ULensFile* InOwner, ELensDataCategory InCategory, ELensCurveAxis InCurveAxis, float InCurveValue, int32 InParameterIndex);
+
+	//~ Begin FRichCurveEditorModel interface
+	virtual void SetKeyPositions(TArrayView<const FKeyHandle> InKeys, TArrayView<const FKeyPosition> InKeyPositions, EPropertyChangeType::Type ChangeType) override;
+	virtual void SetKeyAttributes(TArrayView<const FKeyHandle> InKeys, TArrayView<const FKeyAttributes> InAttributes, EPropertyChangeType::Type ChangeType) override;
+	virtual FText GetKeyLabel() const override;
+	virtual FText GetValueLabel() const override;
+	virtual FText GetValueUnitSuffixLabel() const override;
+	//~ End FRichCurveEditorModel interface
+	
+private:
+	/** The category of the data table being displayed in the curve */
+	ELensDataCategory Category;
+	
+	/** The axis to use as the x-axis of the curve */
+	ELensCurveAxis CurveAxis;
+
+	/** The value to hold fixed on the axis that isn't being plotted */
+	float CurveValue = 0.0;
+
+	/** The index of the data table parameter whose values are plotted on the y-axis of the curve */
+	int32 ParameterIndex = INDEX_NONE;
+};
diff --git a/Engine/Plugins/VirtualProduction/CameraCalibration/Source/CameraCalibrationEditor/Private/AssetEditor/Curves/LensNodalOffsetCurveModel.cpp b/Engine/Plugins/VirtualProduction/CameraCalibration/Source/CameraCalibrationEditor/Private/AssetEditor/Curves/LensNodalOffsetCurveModel.cpp
deleted file mode 100644
diff --git a/Engine/Plugins/VirtualProduction/CameraCalibration/Source/CameraCalibrationEditor/Private/AssetEditor/Curves/LensNodalOffsetCurveModel.h b/Engine/Plugins/VirtualProduction/CameraCalibration/Source/CameraCalibrationEditor/Private/AssetEditor/Curves/LensNodalOffsetCurveModel.h
deleted file mode 100644
diff --git a/Engine/Plugins/VirtualProduction/CameraCalibration/Source/CameraCalibrationEditor/Private/AssetEditor/Curves/LensSTMapCurveModel.cpp b/Engine/Plugins/VirtualProduction/CameraCalibration/Source/CameraCalibrationEditor/Private/AssetEditor/Curves/LensSTMapCurveModel.cpp
deleted file mode 100644
diff --git a/Engine/Plugins/VirtualProduction/CameraCalibration/Source/CameraCalibrationEditor/Private/AssetEditor/Curves/LensSTMapCurveModel.h b/Engine/Plugins/VirtualProduction/CameraCalibration/Source/CameraCalibrationEditor/Private/AssetEditor/Curves/LensSTMapCurveModel.h
deleted file mode 100644
diff --git a/Engine/Plugins/VirtualProduction/CameraCalibration/Source/CameraCalibrationEditor/Private/AssetEditor/SLensDataCategoryListItem.cpp b/Engine/Plugins/VirtualProduction/CameraCalibration/Source/CameraCalibrationEditor/Private/AssetEditor/SLensDataCategoryListItem.cpp
--- a/Engine/Plugins/VirtualProduction/CameraCalibration/Source/CameraCalibrationEditor/Private/AssetEditor/SLensDataCategoryListItem.cpp
+++ b/Engine/Plugins/VirtualProduction/CameraCalibration/Source/CameraCalibrationEditor/Private/AssetEditor/SLensDataCategoryListItem.cpp
@@ -11,8 +11,9 @@
 
 
 
-FLensDataCategoryItem::FLensDataCategoryItem(ULensFile* InLensFile, TWeakPtr<FLensDataCategoryItem> InParent, ELensDataCategory InCategory, FName InLabel)
+FLensDataCategoryItem::FLensDataCategoryItem(ULensFile* InLensFile, TWeakPtr<FLensDataCategoryItem> InParent, ELensDataCategory InCategory, int32 InParameterIndex, FName InLabel)
 	: Category(InCategory)
+	, ParameterIndex(InParameterIndex)
 	, Label(InLabel)
 	, Parent(InParent)
 	, LensFile(InLensFile)
diff --git a/Engine/Plugins/VirtualProduction/CameraCalibration/Source/CameraCalibrationEditor/Private/AssetEditor/SLensDataCategoryListItem.h b/Engine/Plugins/VirtualProduction/CameraCalibration/Source/CameraCalibrationEditor/Private/AssetEditor/SLensDataCategoryListItem.h
--- a/Engine/Plugins/VirtualProduction/CameraCalibration/Source/CameraCalibrationEditor/Private/AssetEditor/SLensDataCategoryListItem.h
+++ b/Engine/Plugins/VirtualProduction/CameraCalibration/Source/CameraCalibrationEditor/Private/AssetEditor/SLensDataCategoryListItem.h
@@ -21,19 +21,20 @@ class FLensDataListItem;
 class FLensDataCategoryItem : public TSharedFromThis<FLensDataCategoryItem>
 {
 public:
-	FLensDataCategoryItem(ULensFile* InLensFile, TWeakPtr<FLensDataCategoryItem> Parent, ELensDataCategory InCategory, FName InLabel);
+	FLensDataCategoryItem(ULensFile* InLensFile, TWeakPtr<FLensDataCategoryItem> Parent, ELensDataCategory InCategory, int32 InParameterIndex, FName InLabel);
 	virtual ~FLensDataCategoryItem() = default;
 
 	/** Makes the widget for its associated row */
 	TSharedRef<ITableRow> MakeTreeRowWidget(const TSharedRef<STableViewBase>& InOwnerTable);
-
-	virtual int32 GetParameterIndex() const { return INDEX_NONE; }
-
+	
 public:
 
 	/** Category this item is associated with */
 	ELensDataCategory Category;
 
+	/** Used to identify which parameter this represents */
+	int32 ParameterIndex = INDEX_NONE;
+	
 	/** Label of this category */
 	FName Label;
 
@@ -67,111 +68,3 @@ private:
 	/** WeakPtr to source data item */
 	TWeakPtr<FLensDataCategoryItem> WeakItem;
 };
-
-/**
- * Distortion parameters category
- */
-class FDistortionParametersCategoryItem : public FLensDataCategoryItem
-{
-public:
-	FDistortionParametersCategoryItem(ULensFile* LensFile, TWeakPtr<FLensDataCategoryItem> Parent, ELensDataCategory InCategory, FName InLabel, int32 InParameterIndex)
-		: FLensDataCategoryItem(LensFile, Parent, InCategory, InLabel)
-		, ParameterIndex(InParameterIndex)
-	{}
-
-	//~ Begin FLensDataCategoryItem interface
-	virtual int32 GetParameterIndex() const override { return ParameterIndex; }
-	//~ End FLensDataCategoryItem interface
-
-public:
-
-	/** Used to identify which distortion parameter this represents */
-	int32 ParameterIndex = 0;
-};
-
-/**
- * Focal length parameters
- */
-class FFocalLengthCategoryItem : public FLensDataCategoryItem
-{
-public:
-	FFocalLengthCategoryItem(ULensFile* LensFile, TWeakPtr<FLensDataCategoryItem> Parent, ELensDataCategory InCategory, FName InLabel, int32 InParameterIndex)
-		: FLensDataCategoryItem(LensFile, Parent, InCategory, InLabel)
-		, ParameterIndex(InParameterIndex)
-	{
-		check(InParameterIndex >= 0 && InParameterIndex < 2);
-	}
-
-	//~ Begin FLensDataCategoryItem interface
-	virtual int32 GetParameterIndex() const { return ParameterIndex; }
-	//~ End FLensDataCategoryItem interface
-
-public:
-
-	/**
-	 * Used to identify which image center parameter this represents
-	 * 0: Fx
-	 * 1: Fy
-	 */
-	int32 ParameterIndex = 0;
-};
-
-/**
- * Image Center parameters
- */
-class FImageCenterCategoryItem : public FLensDataCategoryItem
-{
-public:
-	FImageCenterCategoryItem(ULensFile* LensFile, TWeakPtr<FLensDataCategoryItem> Parent, ELensDataCategory InCategory, FName InLabel, int32 InParameterIndex)
-		: FLensDataCategoryItem(LensFile, Parent, InCategory, InLabel)
-		, ParameterIndex(InParameterIndex)
-	{
-		check(InParameterIndex >= 0 && InParameterIndex < 2);
-	}
-
-	//~ Begin FLensDataCategoryItem interface
-	virtual int32 GetParameterIndex() const { return ParameterIndex; }
-	//~ End FLensDataCategoryItem interface
-
-public:
-
-	/** 
-	 * Used to identify which image center parameter this represents 
-	 * 0: Cx
-	 * 1: Cy
-	 */
-	int32 ParameterIndex = 0;
-};
-
-
-/**
- * Nodal Offset parameters
- */
-class FNodalOffsetCategoryItem : public FLensDataCategoryItem
-{
-public:
-	FNodalOffsetCategoryItem(ULensFile* LensFile, TWeakPtr<FLensDataCategoryItem> Parent, ELensDataCategory InCategory, FName InLabel, int32 InParameterIndex, EAxis::Type InAxis)
-		: FLensDataCategoryItem(LensFile, Parent, InCategory, InLabel)
-		, ParameterIndex(InParameterIndex)
-		, Axis(InAxis)
-	{
-		check(InParameterIndex >= 0 && InParameterIndex < 2);
-	}
-
-	//~ Begin FLensDataCategoryItem interface
-	virtual int32 GetParameterIndex() const { return ParameterIndex; }
-	//~ End FLensDataCategoryItem interface
-
-public:
-
-	/**
-	 * Used to identify which image center parameter this represents
-	 * 0: Location
-	 * 1: Rotation
-	 */
-	int32 ParameterIndex = 0;
-
-	/** Axis for the given parameter */
-	EAxis::Type Axis;
-};
-
diff --git a/Engine/Plugins/VirtualProduction/CameraCalibration/Source/CameraCalibrationEditor/Private/AssetEditor/SLensDataListItem.h b/Engine/Plugins/VirtualProduction/CameraCalibration/Source/CameraCalibrationEditor/Private/AssetEditor/SLensDataListItem.h
--- a/Engine/Plugins/VirtualProduction/CameraCalibration/Source/CameraCalibrationEditor/Private/AssetEditor/SLensDataListItem.h
+++ b/Engine/Plugins/VirtualProduction/CameraCalibration/Source/CameraCalibrationEditor/Private/AssetEditor/SLensDataListItem.h
@@ -30,6 +30,7 @@ public:
 	virtual void OnRemoveRequested() const = 0;
 	virtual TSharedRef<ITableRow> MakeTreeRowWidget(const TSharedRef<STableViewBase>& InOwnerTable) = 0;
 	virtual TOptional<float> GetFocus() const { return TOptional<float>(); }
+	virtual TOptional<float> GetZoom() const { return TOptional<float>(); }
 	virtual int32 GetIndex() const { return INDEX_NONE; }
 	virtual void EditItem() {};
 
@@ -106,6 +107,7 @@ public:
 	virtual TSharedRef<ITableRow> MakeTreeRowWidget(const TSharedRef<STableViewBase>& InOwnerTable) override;
 	virtual void OnRemoveRequested() const override;
 	virtual TOptional<float> GetFocus() const override;
+	virtual TOptional<float> GetZoom() const override { return Zoom; }
 	virtual void EditItem() override;
 	//~ End FLensDataListItem interface
 
diff --git a/Engine/Plugins/VirtualProduction/CameraCalibration/Source/CameraCalibrationEditor/Private/AssetEditor/SLensDataViewer.cpp b/Engine/Plugins/VirtualProduction/CameraCalibration/Source/CameraCalibrationEditor/Private/AssetEditor/SLensDataViewer.cpp
--- a/Engine/Plugins/VirtualProduction/CameraCalibration/Source/CameraCalibrationEditor/Private/AssetEditor/SLensDataViewer.cpp
+++ b/Engine/Plugins/VirtualProduction/CameraCalibration/Source/CameraCalibrationEditor/Private/AssetEditor/SLensDataViewer.cpp
@@ -18,16 +18,15 @@
 #include "SLensDataListItem.h"
 
 #include "Curves/LensDataCurveModel.h"
-#include "Curves/LensDistortionParametersCurveModel.h"
 #include "Curves/LensEncodersCurveModel.h"
-#include "Curves/LensFocalLengthCurveModel.h"
-#include "Curves/LensImageCenterCurveModel.h"
-#include "Curves/LensNodalOffsetCurveModel.h"
-#include "Curves/LensSTMapCurveModel.h"
+#include "Curves/LensMultiAxisCurveModel.h"
 #include "Framework/MultiBox/MultiBoxBuilder.h"
 #include "Misc/MessageDialog.h"
 #include "Widgets/Images/SImage.h"
 #include "Widgets/Input/SButton.h"
+#include "Widgets/Input/SCheckBox.h"
+#include "Widgets/Input/SComboBox.h"
+#include "Widgets/Layout/SSpacer.h"
 
 
 #define LOCTEXT_NAMESPACE "LensDataViewer"
@@ -169,7 +168,28 @@ void SLensDataViewer::Construct(const FArguments& InArgs, ULensFile* InLensFile,
 			+ SHorizontalBox::Slot()
 			.FillWidth(0.6f)
 			[
-				CurvePanel.ToSharedRef()
+				SNew(SOverlay)
+
+				+SOverlay::Slot()
+				[
+					CurvePanel.ToSharedRef()
+				]
+
+				+SOverlay::Slot()
+				[
+					SNew(SVerticalBox)
+
+					+SVerticalBox::Slot()
+					[
+						SNew(SSpacer)
+					]
+
+					+SVerticalBox::Slot()
+					.AutoHeight()
+					[
+						MakeCurveEditorToolbarWidget()
+					]
+				]
 			]
 		]
 	];
@@ -373,6 +393,7 @@ TSharedRef<SWidget> SLensDataViewer::MakeToolbarWidget(TSharedRef<SCameraCalibra
 	ToolBarBuilder.AddWidget(AddPointButton);
 	ToolBarBuilder.AddSeparator();
 	ToolBarBuilder.AddWidget(ClearAllButton);
+	ToolBarBuilder.AddSeparator();
 		
 	ToolBarBuilder.EndSection();
 
@@ -383,6 +404,57 @@ TSharedRef<SWidget> SLensDataViewer::MakeToolbarWidget(TSharedRef<SCameraCalibra
 		];
 }
 
+TSharedRef<SWidget> SLensDataViewer::MakeCurveEditorToolbarWidget()
+{
+	return SNew(SHorizontalBox)
+
+	+SHorizontalBox::Slot()
+	.AutoWidth()
+	.HAlign(HAlign_Center)
+	.VAlign(VAlign_Center)
+	[
+		SNew(SBorder)
+		.BorderImage(FAppStyle::Get().GetBrush("NoBorder"))
+		.Visibility(this, &SLensDataViewer::GetCurveAxisButtonVisibility)
+		[
+			SNew(SStackBox)
+			.Orientation(Orient_Horizontal)
+
+			+SStackBox::Slot()
+			[
+				SNew(SCheckBox)
+				.Style(FAppStyle::Get(), "ToggleButtonCheckBox")
+				.OnCheckStateChanged_Lambda([this](ECheckBoxState InCheckBoxState)
+				{
+					CurveAxisType = ELensCurveAxis::Zoom;
+					RefreshCurve();
+				})
+				.IsChecked_Lambda([this]() { return CurveAxisType == ELensCurveAxis::Zoom ? ECheckBoxState::Checked : ECheckBoxState::Unchecked; })
+				.Padding(4.f)
+				[
+					SNew(STextBlock).Text(LOCTEXT("ZoomAxisLabel", "Zoom"))
+				]
+			]
+			
+			+SStackBox::Slot()
+			[
+				SNew(SCheckBox)
+				.Style(FAppStyle::Get(), "ToggleButtonCheckBox")
+				.OnCheckStateChanged_Lambda([this](ECheckBoxState InCheckBoxState)
+				{
+					CurveAxisType = ELensCurveAxis::Focus;
+					RefreshCurve();
+				})
+				.IsChecked_Lambda([this]() { return CurveAxisType == ELensCurveAxis::Focus ? ECheckBoxState::Checked : ECheckBoxState::Unchecked; })
+				.Padding(4.f)
+				[
+					SNew(STextBlock).Text(LOCTEXT("FocusAxisLabel", "Focus"))
+				]
+			]
+		]
+	];
+}
+
 void SLensDataViewer::OnAddDataPointHandler()
 {
 	const FSimpleDelegate OnDataPointAdded = FSimpleDelegate::CreateSP(this, &SLensDataViewer::OnLensDataPointAdded);
@@ -432,18 +504,24 @@ void SLensDataViewer::RefreshDataCategoriesTree()
 	
 	DataCategories.Reset();
 
-	DataCategories.Add(MakeShared<FLensDataCategoryItem>(LensFile.Get(), nullptr, ELensDataCategory::Focus,  LensDataUtils::EncoderFocusLabel));
-	DataCategories.Add(MakeShared<FLensDataCategoryItem>(LensFile.Get(), nullptr, ELensDataCategory::Iris, LensDataUtils::EncoderIrisLabel));
+	DataCategories.Add(MakeShared<FLensDataCategoryItem>(LensFile.Get(), nullptr, ELensDataCategory::Focus, INDEX_NONE, LensDataUtils::EncoderFocusLabel));
+	DataCategories.Add(MakeShared<FLensDataCategoryItem>(LensFile.Get(), nullptr, ELensDataCategory::Iris, INDEX_NONE, LensDataUtils::EncoderIrisLabel));
 
-	TSharedPtr<FLensDataCategoryItem> FocalLengthCategory = MakeShared<FLensDataCategoryItem>(LensFile.Get(), nullptr, ELensDataCategory::Zoom, LensDataUtils::EncoderZoomLabel);
-	FocalLengthCategory->Children.Add(MakeShared<FFocalLengthCategoryItem>(LensFile.Get(), FocalLengthCategory, ELensDataCategory::Zoom, LensDataUtils::FxLabel, 0));
-	FocalLengthCategory->Children.Add(MakeShared<FFocalLengthCategoryItem>(LensFile.Get(), FocalLengthCategory, ELensDataCategory::Zoom, LensDataUtils::FyLabel, 1));
+	TSharedPtr<FLensDataCategoryItem> FocalLengthCategory = MakeShared<FLensDataCategoryItem>(
+		LensFile.Get(), nullptr, ELensDataCategory::Zoom, FFocalLengthTable::FParameters::Aggregate,
+		LensDataUtils::EncoderZoomLabel);
+	
+	FocalLengthCategory->Children.Add(MakeShared<FLensDataCategoryItem>(LensFile.Get(), FocalLengthCategory, ELensDataCategory::Zoom, FFocalLengthTable::FParameters::Fx, LensDataUtils::FxLabel));
+	FocalLengthCategory->Children.Add(MakeShared<FLensDataCategoryItem>(LensFile.Get(), FocalLengthCategory, ELensDataCategory::Zoom, FFocalLengthTable::FParameters::Fy, LensDataUtils::FyLabel));
 	DataCategories.Add(FocalLengthCategory);
 
 	
 	if (LensFile->DataMode == ELensDataMode::Parameters)
 	{
-		TSharedPtr<FLensDataCategoryItem> DistortionEntry = MakeShared<FLensDataCategoryItem>(LensFile.Get(), nullptr, ELensDataCategory::Distortion, LensDataUtils::DistortionCategoryLabel);
+		TSharedPtr<FLensDataCategoryItem> DistortionEntry = MakeShared<FLensDataCategoryItem>(
+			LensFile.Get(), nullptr, ELensDataCategory::Distortion, FDistortionTable::FParameters::Aggregate,
+			LensDataUtils::DistortionCategoryLabel);
+		
 		DataCategories.Add(DistortionEntry);
 
 		TArray<FText> Parameters;
@@ -455,29 +533,34 @@ void SLensDataViewer::RefreshDataCategoriesTree()
 		for (int32 Index = 0; Index< Parameters.Num(); ++Index)
 		{
 			const FText& Parameter = Parameters[Index];
-			DistortionEntry->Children.Add(MakeShared<FDistortionParametersCategoryItem>(LensFile.Get(), DistortionEntry, ELensDataCategory::Distortion, *Parameter.ToString(), Index));
+			DistortionEntry->Children.Add(MakeShared<FLensDataCategoryItem>(LensFile.Get(), DistortionEntry, ELensDataCategory::Distortion, Index, *Parameter.ToString()));
 		}
 	}
 	else
 	{
-		DataCategories.Add(MakeShared<FLensDataCategoryItem>(LensFile.Get(), nullptr, ELensDataCategory::STMap, LensDataUtils::MapsCategoryLabel));
+		DataCategories.Add(MakeShared<FLensDataCategoryItem>(LensFile.Get(), nullptr, ELensDataCategory::STMap, INDEX_NONE, LensDataUtils::MapsCategoryLabel));
 	}
 
-	TSharedPtr<FLensDataCategoryItem> ImageCenterEntry = MakeShared<FLensDataCategoryItem>(LensFile.Get(), nullptr, ELensDataCategory::ImageCenter, LensDataUtils::ImageCenterCategory);
+	TSharedPtr<FLensDataCategoryItem> ImageCenterEntry = MakeShared<FLensDataCategoryItem>(LensFile.Get(), nullptr, ELensDataCategory::ImageCenter, INDEX_NONE, LensDataUtils::ImageCenterCategory);
 	DataCategories.Add(ImageCenterEntry);
-	ImageCenterEntry->Children.Add(MakeShared<FImageCenterCategoryItem>(LensFile.Get(), ImageCenterEntry, ELensDataCategory::ImageCenter, LensDataUtils::CxLabel, 0));
-	ImageCenterEntry->Children.Add(MakeShared<FImageCenterCategoryItem>(LensFile.Get(), ImageCenterEntry, ELensDataCategory::ImageCenter, LensDataUtils::CyLabel, 1));
+	ImageCenterEntry->Children.Add(MakeShared<FLensDataCategoryItem>(LensFile.Get(), ImageCenterEntry, ELensDataCategory::ImageCenter, FImageCenterTable::FParameters::Cx, LensDataUtils::CxLabel));
+	ImageCenterEntry->Children.Add(MakeShared<FLensDataCategoryItem>(LensFile.Get(), ImageCenterEntry, ELensDataCategory::ImageCenter, FImageCenterTable::FParameters::Cy, LensDataUtils::CyLabel));
 
-	TSharedPtr<FLensDataCategoryItem> NodalOffsetCategory = MakeShared<FLensDataCategoryItem>(LensFile.Get(), nullptr, ELensDataCategory::NodalOffset, LensDataUtils::NodalOffsetCategoryLabel);
+	TSharedPtr<FLensDataCategoryItem> NodalOffsetCategory = MakeShared<FLensDataCategoryItem>(LensFile.Get(), nullptr, ELensDataCategory::NodalOffset, INDEX_NONE, LensDataUtils::NodalOffsetCategoryLabel);
 	DataCategories.Add(NodalOffsetCategory);
 
-	NodalOffsetCategory->Children.Add(MakeShared<FNodalOffsetCategoryItem>(LensFile.Get(), NodalOffsetCategory, ELensDataCategory::NodalOffset, LensDataUtils::LocationXLabel, 0, EAxis::X));
-	NodalOffsetCategory->Children.Add(MakeShared<FNodalOffsetCategoryItem>(LensFile.Get(), NodalOffsetCategory, ELensDataCategory::NodalOffset, LensDataUtils::LocationYLabel, 0, EAxis::Y));
-	NodalOffsetCategory->Children.Add(MakeShared<FNodalOffsetCategoryItem>(LensFile.Get(), NodalOffsetCategory, ELensDataCategory::NodalOffset, LensDataUtils::LocationZLabel, 0, EAxis::Z));
-	NodalOffsetCategory->Children.Add(MakeShared<FNodalOffsetCategoryItem>(LensFile.Get(), NodalOffsetCategory, ELensDataCategory::NodalOffset, LensDataUtils::RotationXLabel, 1, EAxis::X));
-	NodalOffsetCategory->Children.Add(MakeShared<FNodalOffsetCategoryItem>(LensFile.Get(), NodalOffsetCategory, ELensDataCategory::NodalOffset, LensDataUtils::RotationYLabel, 1, EAxis::Y));
-	NodalOffsetCategory->Children.Add(MakeShared<FNodalOffsetCategoryItem>(LensFile.Get(), NodalOffsetCategory, ELensDataCategory::NodalOffset, LensDataUtils::RotationZLabel, 1, EAxis::Z));
-
+	{
+		using FParameters = FNodalOffsetTable::FParameters;
+		const ELensDataCategory Category = ELensDataCategory::NodalOffset;
+		
+		NodalOffsetCategory->Children.Add(MakeShared<FLensDataCategoryItem>(LensFile.Get(), NodalOffsetCategory, Category, FParameters::Compose(FParameters::Location, EAxis::X), LensDataUtils::LocationXLabel));
+		NodalOffsetCategory->Children.Add(MakeShared<FLensDataCategoryItem>(LensFile.Get(), NodalOffsetCategory, Category, FParameters::Compose(FParameters::Location, EAxis::Y), LensDataUtils::LocationYLabel));
+		NodalOffsetCategory->Children.Add(MakeShared<FLensDataCategoryItem>(LensFile.Get(), NodalOffsetCategory, Category, FParameters::Compose(FParameters::Location, EAxis::Z), LensDataUtils::LocationZLabel));
+		NodalOffsetCategory->Children.Add(MakeShared<FLensDataCategoryItem>(LensFile.Get(), NodalOffsetCategory, Category, FParameters::Compose(FParameters::Rotation, EAxis::X), LensDataUtils::RotationXLabel));
+		NodalOffsetCategory->Children.Add(MakeShared<FLensDataCategoryItem>(LensFile.Get(), NodalOffsetCategory, Category, FParameters::Compose(FParameters::Rotation, EAxis::Y), LensDataUtils::RotationYLabel));
+		NodalOffsetCategory->Children.Add(MakeShared<FLensDataCategoryItem>(LensFile.Get(), NodalOffsetCategory, Category, FParameters::Compose(FParameters::Rotation, EAxis::Z), LensDataUtils::RotationZLabel));
+	}
+	
 	TreeView->RequestTreeRefresh();
 }
 
@@ -526,31 +609,31 @@ void SLensDataViewer::RefreshDataEntriesTree()
 			case ELensDataCategory::Zoom:
 			{
 				const TConstArrayView<FFocalLengthFocusPoint> FocusPoints = LensFile->FocalLengthTable.GetFocusPoints();
-				LensDataUtils::MakeFocusEntries(LensFile.Get(), CategoryItem->Category, CategoryItem->GetParameterIndex(), FocusPoints, DataEntries, DataChangedCallback);
+				LensDataUtils::MakeFocusEntries(LensFile.Get(), CategoryItem->Category, CategoryItem->ParameterIndex, FocusPoints, DataEntries, DataChangedCallback);
 				break;
 			}
 			case ELensDataCategory::Distortion:
 			{
 				const TConstArrayView<FDistortionFocusPoint> FocusPoints = LensFile->DistortionTable.GetFocusPoints();
-				LensDataUtils::MakeFocusEntries(LensFile.Get(), CategoryItem->Category, CategoryItem->GetParameterIndex(), FocusPoints, DataEntries, DataChangedCallback);
+				LensDataUtils::MakeFocusEntries(LensFile.Get(), CategoryItem->Category, CategoryItem->ParameterIndex, FocusPoints, DataEntries, DataChangedCallback);
 				break;
 			}
 			case ELensDataCategory::ImageCenter:
 			{
 				const TConstArrayView<FImageCenterFocusPoint> FocusPoints = LensFile->ImageCenterTable.GetFocusPoints();
-				LensDataUtils::MakeFocusEntries(LensFile.Get(), CategoryItem->Category, CategoryItem->GetParameterIndex(), FocusPoints, DataEntries, DataChangedCallback);
+				LensDataUtils::MakeFocusEntries(LensFile.Get(), CategoryItem->Category, CategoryItem->ParameterIndex, FocusPoints, DataEntries, DataChangedCallback);
 				break;
 			}
 			case ELensDataCategory::NodalOffset:
 			{
 				const TConstArrayView<FNodalOffsetFocusPoint> Points = LensFile->NodalOffsetTable.GetFocusPoints();
-				LensDataUtils::MakeFocusEntries(LensFile.Get(), CategoryItem->Category, CategoryItem->GetParameterIndex(), Points, DataEntries, DataChangedCallback);
+				LensDataUtils::MakeFocusEntries(LensFile.Get(), CategoryItem->Category, CategoryItem->ParameterIndex, Points, DataEntries, DataChangedCallback);
 				break;
 			}
 			case ELensDataCategory::STMap:
 			{
 				const TConstArrayView<FSTMapFocusPoint> Points = LensFile->STMapTable.GetFocusPoints();
-				LensDataUtils::MakeFocusEntries(LensFile.Get(), CategoryItem->Category, CategoryItem->GetParameterIndex(), Points, DataEntries, DataChangedCallback);
+				LensDataUtils::MakeFocusEntries(LensFile.Get(), CategoryItem->Category, CategoryItem->ParameterIndex, Points, DataEntries, DataChangedCallback);
 				break;
 			}
 		}
@@ -581,93 +664,43 @@ void SLensDataViewer::RefreshCurve() const
 	TSharedPtr<FLensDataCategoryItem> CategoryItem = GetDataCategorySelection();
 	if (CategoryItem.IsValid())
 	{
-		const TSharedPtr<FLensDataListItem> CurrentDataItem = GetSelectedDataEntry();
-		switch (CategoryItem->Category)
+		const ELensDataCategory Category = CategoryItem->Category;
+		if (Category == ELensDataCategory::Focus || Category == ELensDataCategory::Iris)
 		{
-			case ELensDataCategory::Focus:
-			{
-				NewCurve = MakeUnique<FLensEncodersCurveModel>(LensFile.Get(), EEncoderType::Focus);
-				break;
-			}
-			case ELensDataCategory::Iris:
-			{
-				NewCurve = MakeUnique<FLensEncodersCurveModel>(LensFile.Get(), EEncoderType::Iris);
-				break;
-			}
-			case ELensDataCategory::Zoom:
-			{
-				if (CurrentDataItem)
-				{
-					const TOptional<float> Focus = CurrentDataItem->GetFocus();
-					if(Focus.IsSet())
-					{
-						const int32 ParameterIndex = CategoryItem->GetParameterIndex();
-						NewCurve = MakeUnique<FLensFocalLengthCurveModel>(LensFile.Get(), Focus.GetValue(), ParameterIndex);
-					}
-				}
-				break;
-			}
-			case ELensDataCategory::Distortion:
+			const EEncoderType EncoderType = Category == ELensDataCategory::Focus ? EEncoderType::Focus : EEncoderType::Iris;
+			NewCurve = MakeUnique<FLensEncodersCurveModel>(LensFile.Get(), EncoderType);
+		}
+		else
+		{
+			const TSharedPtr<FLensDataListItem> CurrentDataItem = GetSelectedDataEntry();
+			if (CurrentDataItem)
 			{
-				if (CurrentDataItem)
+				TOptional<float> CurveValue = TOptional<float>();
+				if (CurveAxisType == ELensCurveAxis::Focus)
 				{
-					const TOptional<float> Focus = CurrentDataItem->GetFocus();
-					if(Focus.IsSet())
+					CurveValue = CurrentDataItem->GetZoom();
+					if (!CurveValue.IsSet() && CurrentDataItem->Children.Num())
 					{
-						const int32 ParameterIndex = CategoryItem->GetParameterIndex();
-						NewCurve = MakeUnique<FLensDistortionParametersCurveModel>(LensFile.Get(), Focus.GetValue(), ParameterIndex);
+						// If the current data item does not have a zoom value, check to see if it has any children, and if so,
+						// use the zoom from the first child
+						CurveValue = CurrentDataItem->Children[0]->GetZoom();
 					}
 				}
-				break;
-			}
-			case ELensDataCategory::ImageCenter:
-			{
-				if (CurrentDataItem)
+				else
 				{
-					const TOptional<float> Focus = CurrentDataItem->GetFocus();
-					const int32 ParameterIndex = CategoryItem->GetParameterIndex();
-					if(Focus.IsSet() && ParameterIndex != INDEX_NONE)
-					{
-						NewCurve = MakeUnique<FLensImageCenterCurveModel>(LensFile.Get(), Focus.GetValue(), ParameterIndex);
-					}
+					CurveValue = CurrentDataItem->GetFocus();
 				}
-				break;
-			}
-			case ELensDataCategory::NodalOffset:
-			{
-				if (CurrentDataItem)
+
+				if (CurveValue.IsSet())
 				{
-					TSharedPtr<FNodalOffsetCategoryItem> NodalCategory = StaticCastSharedPtr<FNodalOffsetCategoryItem>(CategoryItem);
-					const TOptional<float> Focus = CurrentDataItem->GetFocus();
-					const int32 ParameterIndex = CategoryItem->GetParameterIndex();
-					if(Focus.IsSet() && ParameterIndex != INDEX_NONE)
-					{
-						NewCurve = MakeUnique<FLensNodalOffsetCurveModel>(LensFile.Get(), Focus.GetValue(), ParameterIndex, NodalCategory->Axis);
-					}
+					NewCurve = MakeUnique<FLensDataMultiAxisCurveModel>(LensFile.Get(), CategoryItem->Category, CurveAxisType, CurveValue.GetValue(), CategoryItem->ParameterIndex);
 				}
-				break;
-			}
-			case ELensDataCategory::STMap:
-			{
-					if (CurrentDataItem)
-					{
-						const TOptional<float> Focus = CurrentDataItem->GetFocus();
-						if(Focus.IsSet())
-						{
-							NewCurve = MakeUnique<FLensSTMapCurveModel>(LensFile.Get(), Focus.GetValue());
-						}
-					}
-				break;
-			}
-			default:
-			{
-				break;
 			}
 		}
 	}
 
 	//If a curve was setup, add it to the editor
-	if (NewCurve)
+	if (NewCurve && NewCurve->IsValid())
 	{
 		NewCurve->SetShortDisplayName(FText::FromName(CategoryItem->Label));
 		const UCameraCalibrationEditorSettings* EditorSettings = GetDefault<UCameraCalibrationEditorSettings>();
@@ -693,7 +726,7 @@ void SLensDataViewer::RefreshTimeSlider() const
 	}
 
 	const bool bHasParameterIndex = CategoryItem->Category == ELensDataCategory::ImageCenter || CategoryItem->Category == ELensDataCategory::NodalOffset;
-	const int32 ParameterIndex = CategoryItem->GetParameterIndex();
+	const int32 ParameterIndex = CategoryItem->ParameterIndex;
 
 	// Reset selection if no selection of the curve for ImageCenter or NodalOffset
 	if (bHasParameterIndex && ParameterIndex == INDEX_NONE)
@@ -756,6 +789,20 @@ void SLensDataViewer::UpdateDataSelection(const TSharedPtr<FLensDataListItem>& P
 	}
 }
 
+EVisibility SLensDataViewer::GetCurveAxisButtonVisibility() const
+{
+	TSharedPtr<FLensDataCategoryItem> CategoryItem = GetDataCategorySelection();
+	if (CategoryItem.IsValid())
+	{
+		if (CategoryItem->Category != ELensDataCategory::Focus && CategoryItem->Category != ELensDataCategory::Iris)
+		{
+			return EVisibility::Visible;
+		}
+	}
+
+	return EVisibility::Hidden;
+}
+
 #undef LOCTEXT_NAMESPACE /* LensDataViewer */
 
 
diff --git a/Engine/Plugins/VirtualProduction/CameraCalibration/Source/CameraCalibrationEditor/Private/AssetEditor/SLensDataViewer.h b/Engine/Plugins/VirtualProduction/CameraCalibration/Source/CameraCalibrationEditor/Private/AssetEditor/SLensDataViewer.h
--- a/Engine/Plugins/VirtualProduction/CameraCalibration/Source/CameraCalibrationEditor/Private/AssetEditor/SLensDataViewer.h
+++ b/Engine/Plugins/VirtualProduction/CameraCalibration/Source/CameraCalibrationEditor/Private/AssetEditor/SLensDataViewer.h
@@ -19,6 +19,7 @@ class SCameraCalibrationCurveEditorPanel;
 class FCameraCalibrationStepsController;
 class FCameraCalibrationTimeSliderController;
 enum class ELensDataChangedReason : uint8;
+enum class ELensCurveAxis : uint8;
 
 /** Widget used to display data from the LensFile */
 class SLensDataViewer : public SCompoundWidget, public FSelfRegisteringEditorUndoClient
@@ -71,6 +72,9 @@ private:
 	/** Makes the Toolbar with data manipulation buttons */
 	TSharedRef<SWidget> MakeToolbarWidget(TSharedRef<SCameraCalibrationCurveEditorPanel> InEditorPanel);
 
+	/** Makes the toolbar for the curve editor, which sits on top of the curve editor */
+	TSharedRef<SWidget> MakeCurveEditorToolbarWidget();
+	
 	/** Curve Editor add point delegate handler*/
 	void OnAddDataPointHandler();
 	
@@ -107,6 +111,9 @@ private:
 	/** Used to keep same selected data point when rebuilding tree */
 	void UpdateDataSelection(const TSharedPtr<FLensDataListItem>& PreviousSelection);
 
+	/** Gets whether the axis selection buttons should be visible in the curve editor toolbar */
+	EVisibility GetCurveAxisButtonVisibility() const;
+
 private:
 	
 	/** Data category TreeView */
@@ -139,6 +146,9 @@ private:
 	/** Child class of curve editor panel */
 	TSharedPtr<SCameraCalibrationCurveEditorPanel> CurvePanel;
 
+	/** The current axis type to display in the curve editor */
+	ELensCurveAxis CurveAxisType = (ELensCurveAxis)0;
+	
 	/** Evaluated FIZ for the current frame */
 	TAttribute<FCachedFIZData> CachedFIZ;
 
diff --git a/Engine/Plugins/VirtualProduction/CameraCalibrationCore/Source/CameraCalibrationCore/Private/LensFile.cpp b/Engine/Plugins/VirtualProduction/CameraCalibrationCore/Source/CameraCalibrationCore/Private/LensFile.cpp
--- a/Engine/Plugins/VirtualProduction/CameraCalibrationCore/Source/CameraCalibrationCore/Private/LensFile.cpp
+++ b/Engine/Plugins/VirtualProduction/CameraCalibrationCore/Source/CameraCalibrationCore/Private/LensFile.cpp
@@ -1343,6 +1343,11 @@ const FBaseLensTable* ULensFile::GetDataTable(ELensDataCategory InDataCategory)
 	}
 }
 
+FBaseLensTable* ULensFile::GetDataTable(ELensDataCategory InDataCategory)
+{
+	return const_cast<FBaseLensTable*>(const_cast<const ULensFile*>(this)->GetDataTable(InDataCategory));
+}
+
 void ULensFile::PostInitProperties()
 {
 	Super::PostInitProperties();
diff --git a/Engine/Plugins/VirtualProduction/CameraCalibrationCore/Source/CameraCalibrationCore/Private/Tables/DistortionParametersTable.cpp b/Engine/Plugins/VirtualProduction/CameraCalibrationCore/Source/CameraCalibrationCore/Private/Tables/DistortionParametersTable.cpp
--- a/Engine/Plugins/VirtualProduction/CameraCalibrationCore/Source/CameraCalibrationCore/Private/Tables/DistortionParametersTable.cpp
+++ b/Engine/Plugins/VirtualProduction/CameraCalibrationCore/Source/CameraCalibrationCore/Private/Tables/DistortionParametersTable.cpp
@@ -113,7 +113,7 @@ UScriptStruct* FDistortionTable::GetScriptStruct() const
 	return StaticStruct();
 }
 
-bool FDistortionTable::BuildParameterCurve(float InFocus, int32 ParameterIndex, FRichCurve& OutCurve) const
+bool FDistortionTable::BuildParameterCurveAtFocus(float InFocus, int32 ParameterIndex, FRichCurve& OutCurve) const
 {
 	if (const FDistortionFocusPoint* ThisFocusPoints = GetFocusPoint(InFocus))
 	{
@@ -138,6 +138,114 @@ bool FDistortionTable::BuildParameterCurve(float InFocus, int32 ParameterIndex,
 	return false;
 }
 
+bool FDistortionTable::BuildParameterCurveAtZoom(float InZoom, int32 InParameterIndex, FRichCurve& OutCurve) const
+{
+	for (const FDistortionFocusPoint& FocusPoint : FocusPoints)
+	{
+		FDistortionInfo ZoomPoint;
+		if (FocusPoint.GetPoint(InZoom, ZoomPoint))
+		{
+			float Value = 0.0;
+			if (ZoomPoint.Parameters.IsValidIndex(InParameterIndex))
+			{
+				Value = ZoomPoint.Parameters[InParameterIndex];
+			}
+			else
+			{
+				FKeyHandle BlendPointKey = FocusPoint.MapBlendingCurve.FindKey(InZoom);
+				Value = FocusPoint.MapBlendingCurve.GetKeyValue(BlendPointKey);
+			}
+			
+			const FKeyHandle NewKeyHandle = OutCurve.AddKey(FocusPoint.Focus, Value);
+			FRichCurveKey& NewKey = OutCurve.GetKey(NewKeyHandle);
+			NewKey.TangentMode = ERichCurveTangentMode::RCTM_None;
+			NewKey.InterpMode = ERichCurveInterpMode::RCIM_Linear;
+		}
+	}
+
+	return true;
+}
+
+void FDistortionTable::SetParameterCurveKeysAtFocus(float InFocus, int32 InParameterIndex, const FRichCurve& InSourceCurve, TArrayView<const FKeyHandle> InKeys)
+{
+	if (FDistortionFocusPoint* FocusPoint = GetFocusPoint(InFocus))
+	{
+		for (int32 Index = 0; Index < InKeys.Num(); ++Index)
+		{
+			const FKeyHandle Handle = InKeys[Index];
+			const int32 KeyIndex = InSourceCurve.GetIndexSafe(Handle);
+			if (KeyIndex != INDEX_NONE)
+			{
+				//We can't move keys on the time axis so our indices should match
+				const FRichCurveKey& Key = InSourceCurve.GetKey(Handle);
+
+				if (InParameterIndex == FParameters::Aggregate)
+				{
+					FocusPoint->MapBlendingCurve.Keys[KeyIndex] = Key;
+				}
+				else
+				{
+					FocusPoint->SetParameterValue(KeyIndex, Key.Time, InParameterIndex, Key.Value);
+				}
+			}
+		}
+	}
+}
+
+void FDistortionTable::SetParameterCurveKeysAtZoom(float InZoom, int32 InParameterIndex, const FRichCurve& InSourceCurve, TArrayView<const FKeyHandle> InKeys)
+{
+	for (const FKeyHandle& KeyHandle : InKeys)
+	{
+		// Assume the focus keys are put into the source curve in the same order as they are stored internally
+		const int32 KeyIndex = InSourceCurve.GetIndexSafe(KeyHandle);
+		if (KeyIndex != INDEX_NONE)
+		{
+			if (ensure(FocusPoints.IsValidIndex(KeyIndex)))
+			{
+				//We can't move keys on the time axis so our indices should match
+				const FRichCurveKey& Key = InSourceCurve.GetKey(KeyHandle);
+				FDistortionFocusPoint& FocusPoint = FocusPoints[KeyIndex];
+				
+				if (InParameterIndex == FParameters::Aggregate)
+				{
+					FocusPoint.MapBlendingCurve.Keys[KeyIndex] = Key;
+				}
+				else
+				{
+					FDistortionInfo ZoomPoint;
+					if (!FocusPoint.GetPoint(InZoom, ZoomPoint))
+					{
+						continue;
+					}
+					
+					ZoomPoint.Parameters[InParameterIndex] = Key.Value;
+					FocusPoint.SetPoint(InZoom, ZoomPoint);
+				}
+			}
+		}
+	}
+}
+
+bool FDistortionTable::CanEditCurveKeyPositions(int32 InParameterIndex) const
+{
+	return InParameterIndex != FParameters::Aggregate;
+}
+
+bool FDistortionTable::CanEditCurveKeyAttributes(int32 InParameterIndex) const
+{
+	return InParameterIndex == FParameters::Aggregate;
+}
+
+FText FDistortionTable::GetParameterValueLabel(int32 InParameterIndex) const
+{
+	if (InParameterIndex != FParameters::Aggregate)
+	{
+		return NSLOCTEXT("FDistortionTable", "ParameterValueLabel", "(unitless)");
+	}
+	
+	return FText();
+}
+
 const FDistortionFocusPoint* FDistortionTable::GetFocusPoint(float InFocus, float InputTolerance) const
 {
 	return FocusPoints.FindByPredicate([InFocus, InputTolerance](const FDistortionFocusPoint& Point) { return FMath::IsNearlyEqual(Point.Focus, InFocus, InputTolerance); });
diff --git a/Engine/Plugins/VirtualProduction/CameraCalibrationCore/Source/CameraCalibrationCore/Private/Tables/FocalLengthTable.cpp b/Engine/Plugins/VirtualProduction/CameraCalibrationCore/Source/CameraCalibrationCore/Private/Tables/FocalLengthTable.cpp
--- a/Engine/Plugins/VirtualProduction/CameraCalibrationCore/Source/CameraCalibrationCore/Private/Tables/FocalLengthTable.cpp
+++ b/Engine/Plugins/VirtualProduction/CameraCalibrationCore/Source/CameraCalibrationCore/Private/Tables/FocalLengthTable.cpp
@@ -175,25 +175,238 @@ UScriptStruct* FFocalLengthTable::GetScriptStruct() const
 	return StaticStruct();
 }
 
-bool FFocalLengthTable::BuildParameterCurve(float InFocus, int32 ParameterIndex, FRichCurve& OutCurve) const
+bool FFocalLengthTable::BuildParameterCurveAtFocus(float InFocus, int32 ParameterIndex, FRichCurve& OutCurve) const
 {
-	if(ParameterIndex >= 0 && ParameterIndex < 2)
+	if (!FParameters::IsValidOrAggregate(ParameterIndex))
 	{
-		if(const FFocalLengthFocusPoint* FocusPoint = GetFocusPoint(InFocus))
+		return false;
+	}
+	
+	const FFocalLengthFocusPoint* FocusPoint = GetFocusPoint(InFocus);
+	if (!FocusPoint)
+	{
+		return false;
+	}
+
+	if (ParameterIndex == FParameters::Aggregate)
+	{
+		// The aggregate curve is just the x curve scaled by the lens file's sensor
+		const FRichCurve& ActiveCurve = FocusPoint->Fx;
+		auto HandleIter = ActiveCurve.GetKeyHandleIterator();
+		for (const FRichCurveKey& Key : ActiveCurve.GetConstRefOfKeys())
+		{
+			ULensFile* LensFilePtr = GetLensFile();
+			const float Scale = LensFilePtr ? LensFilePtr->LensInfo.SensorDimensions.X : 1.0f;
+			OutCurve.AddKey(Key.Time, Key.Value * Scale, false, *HandleIter);
+
+			FRichCurveKey& NewKey = OutCurve.GetKey(*HandleIter);
+			NewKey.TangentMode = Key.TangentMode;
+			NewKey.InterpMode = Key.InterpMode;
+			NewKey.ArriveTangent = Key.ArriveTangent * Scale;
+			NewKey.LeaveTangent = Key.LeaveTangent * Scale;
+			++HandleIter;
+		}
+
+		return true;
+	}
+	else if (ParameterIndex == FParameters::Fx)
+	{
+		OutCurve = FocusPoint->Fx;
+		return true;
+	}
+	else if (ParameterIndex == FParameters::Fy)
+	{
+		OutCurve = FocusPoint->Fy;
+		return true;
+	}
+
+	return false;
+}
+
+bool FFocalLengthTable::BuildParameterCurveAtZoom(float InZoom, int32 InParameterIndex, FRichCurve& OutCurve) const
+{
+	for (const FFocalLengthFocusPoint& FocusPoint : FocusPoints)
+	{
+		FFocalLengthInfo ZoomPoint;
+		if (FocusPoint.GetPoint(InZoom, ZoomPoint))
 		{
-			if(ParameterIndex == 0)
+			float Value = ZoomPoint.FxFy.X;
+			if (FParameters::IsValid(InParameterIndex))
 			{
-				OutCurve = FocusPoint->Fx;
+				Value = ZoomPoint.FxFy[InParameterIndex];
 			}
 			else
 			{
-				OutCurve = FocusPoint->Fy;
+				if (ULensFile* LensFilePtr = GetLensFile())
+				{
+					Value = ZoomPoint.FxFy.X * LensFilePtr->LensInfo.SensorDimensions.X;
+				}
 			}
-			return true;
-		}	
+			
+			const FKeyHandle NewKeyHandle = OutCurve.AddKey(FocusPoint.Focus, Value);
+			FRichCurveKey& NewKey = OutCurve.GetKey(NewKeyHandle);
+			NewKey.TangentMode = ERichCurveTangentMode::RCTM_None;
+			NewKey.InterpMode = ERichCurveInterpMode::RCIM_Linear;
+		}
+	}
+
+	return true;
+}
+
+void FFocalLengthTable::SetParameterCurveKeysAtFocus(float InFocus, int32 InParameterIndex, const FRichCurve& InSourceCurve, TArrayView<const FKeyHandle> InKeys)
+{
+	FFocalLengthFocusPoint* FocusPoint = GetFocusPoint(InFocus);
+	if (!FocusPoint)
+	{
+		return;
 	}
 	
-	return false;
+	FRichCurve* ActiveCurve = nullptr;
+	float Scale = 1.0f;
+	int32 FxFyIndex = InParameterIndex;
+	if (InParameterIndex == FParameters::Aggregate)
+	{
+		ActiveCurve = &FocusPoint->Fx;
+		
+		if (ULensFile* LensFilePtr = GetLensFile())
+		{
+			Scale = 1.0f / LensFilePtr->LensInfo.SensorDimensions.X;
+		}
+		
+		FxFyIndex = 0; //mm focal length curve changes Fx
+	}
+	else if (InParameterIndex == FParameters::Fx)
+	{
+		ActiveCurve = &FocusPoint->Fx;
+	}
+	else if (InParameterIndex == FParameters::Fy)
+	{
+		ActiveCurve = &FocusPoint->Fy;
+	}
+
+	if (!ActiveCurve)
+	{
+		return;
+	}
+
+	for (const FKeyHandle& KeyHandle : InKeys)
+	{
+		const int32 KeyIndex = InSourceCurve.GetIndexSafe(KeyHandle);
+		if (KeyIndex != INDEX_NONE)
+		{
+			if(ensure(ActiveCurve->Keys.IsValidIndex(KeyIndex) && FocusPoint->ZoomPoints.IsValidIndex(KeyIndex)))
+			{
+				const FRichCurveKey& SourceKey = InSourceCurve.GetKey(KeyHandle);
+				FRichCurveKey& DestinationKey = ActiveCurve->Keys[KeyIndex];
+				
+				DestinationKey.Value = SourceKey.Value * Scale;
+				DestinationKey.InterpMode = SourceKey.InterpMode;
+				DestinationKey.ArriveTangent = SourceKey.ArriveTangent * Scale;
+				DestinationKey.LeaveTangent = SourceKey.LeaveTangent * Scale;
+				DestinationKey.TangentMode = SourceKey.TangentMode;
+				
+				FocusPoint->ZoomPoints[KeyIndex].FocalLengthInfo.FxFy[FxFyIndex] = SourceKey.Value * Scale;
+			}
+		}
+	}	
+
+	ActiveCurve->AutoSetTangents();
+}
+
+void FFocalLengthTable::SetParameterCurveKeysAtZoom(float InZoom, int32 InParameterIndex, const FRichCurve& InSourceCurve, TArrayView<const FKeyHandle> InKeys)
+{
+	if (!FParameters::IsValidOrAggregate(InParameterIndex))
+	{
+		return;
+	}
+	
+	for (const FKeyHandle& KeyHandle : InKeys)
+	{
+		// Assume the focus keys are put into the source curve in the same order as they are stored internally
+		const int32 KeyIndex = InSourceCurve.GetIndexSafe(KeyHandle);
+		if (KeyIndex != INDEX_NONE)
+		{
+			if (ensure(FocusPoints.IsValidIndex(KeyIndex)))
+			{
+				FFocalLengthFocusPoint& FocusPoint = FocusPoints[KeyIndex];
+				FFocalLengthInfo ZoomPoint;
+				if (!FocusPoint.GetPoint(InZoom, ZoomPoint))
+				{
+					continue;
+				}
+				
+				float Scale = 1.0f;
+				int32 FxFyIndex = InParameterIndex;
+				if (InParameterIndex == FParameters::Aggregate)
+				{
+					if (ULensFile* LensFilePtr = GetLensFile())
+					{
+						Scale = 1.0f / LensFilePtr->LensInfo.SensorDimensions.X;
+					}
+		
+					FxFyIndex = 0; //mm focal length curve changes Fx
+				}
+
+				ZoomPoint.FxFy[FxFyIndex] = InSourceCurve.GetKeyValue(KeyHandle) * Scale;
+				FocusPoint.SetPoint(InZoom, ZoomPoint);
+
+				if (InParameterIndex == FParameters::Fy)
+				{
+					FocusPoint.Fy.AutoSetTangents();
+				}
+				else
+				{
+					FocusPoint.Fx.AutoSetTangents();
+				}
+			}
+		}
+	}
+}
+
+TRange<double> FFocalLengthTable::GetCurveKeyPositionRange(int32 InParameterIndex) const
+{
+	TRange<double> Range = FBaseLensTable::GetCurveKeyPositionRange(InParameterIndex);
+	
+	if (!FParameters::IsValidOrAggregate(InParameterIndex))
+	{
+		return Range;
+	}
+
+	if (InParameterIndex == FParameters::Aggregate)
+	{
+		Range.SetLowerBoundValue(1.0);
+	}
+	else if (ULensFile* LensFilePtr = GetLensFile())
+	{
+		Range.SetLowerBoundValue(1.0 / LensFilePtr->LensInfo.SensorDimensions[InParameterIndex]);
+	}
+
+	return Range;
+}
+
+FText FFocalLengthTable::GetParameterValueLabel(int32 InParameterIndex) const
+{
+	if (!FParameters::IsValidOrAggregate(InParameterIndex))
+	{
+		return FText();
+	}
+	
+	if (InParameterIndex == FParameters::Aggregate)
+	{
+		return NSLOCTEXT("FFocalLengthTable", "ParameterValueMMLabel", "(mm)");
+	}
+	
+	return NSLOCTEXT("FFocalLengthTable", "ParameterValueNormalizedLabel", "(normalized)");
+}
+
+FText FFocalLengthTable::GetParameterValueUnitLabel(int32 InParameterIndex) const
+{
+	if (InParameterIndex == FParameters::Aggregate)
+	{
+		return NSLOCTEXT("FFocalLengthTable", "ParameterUnitLabel", "mm");
+	}
+	
+	return FText();
 }
 
 const FFocalLengthFocusPoint* FFocalLengthTable::GetFocusPoint(float InFocus, float InputTolerance) const
diff --git a/Engine/Plugins/VirtualProduction/CameraCalibrationCore/Source/CameraCalibrationCore/Private/Tables/ImageCenterTable.cpp b/Engine/Plugins/VirtualProduction/CameraCalibrationCore/Source/CameraCalibrationCore/Private/Tables/ImageCenterTable.cpp
--- a/Engine/Plugins/VirtualProduction/CameraCalibrationCore/Source/CameraCalibrationCore/Private/Tables/ImageCenterTable.cpp
+++ b/Engine/Plugins/VirtualProduction/CameraCalibrationCore/Source/CameraCalibrationCore/Private/Tables/ImageCenterTable.cpp
@@ -118,25 +118,127 @@ UScriptStruct* FImageCenterTable::GetScriptStruct() const
 	return StaticStruct();
 }
 
-bool FImageCenterTable::BuildParameterCurve(float InFocus, int32 ParameterIndex, FRichCurve& OutCurve) const
+bool FImageCenterTable::BuildParameterCurveAtFocus(float InFocus, int32 ParameterIndex, FRichCurve& OutCurve) const
 {
-	if(ParameterIndex >= 0 && ParameterIndex < 2)
+	if (!FParameters::IsValid(ParameterIndex))
 	{
-		if(const FImageCenterFocusPoint* FocusPoint = GetFocusPoint(InFocus))
+		return false;
+	}
+	
+	if (const FImageCenterFocusPoint* FocusPoint = GetFocusPoint(InFocus))
+	{
+		if (ParameterIndex == FParameters::Cx)
+		{
+			OutCurve = FocusPoint->Cx;
+		}
+		else
+		{
+			OutCurve = FocusPoint->Cy;
+		}
+		
+		return true;
+	}
+
+	return false;
+}
+
+bool FImageCenterTable::BuildParameterCurveAtZoom(float InZoom, int32 InParameterIndex, FRichCurve& OutCurve) const
+{
+	if (!FParameters::IsValid(InParameterIndex))
+	{
+		return false;
+	}
+	
+	for (const FImageCenterFocusPoint& FocusPoint : FocusPoints)
+	{
+		FImageCenterInfo ZoomPoint;
+		if (FocusPoint.GetPoint(InZoom, ZoomPoint))
+		{
+			const float Value = ZoomPoint.PrincipalPoint[InParameterIndex];
+			const FKeyHandle NewKeyHandle = OutCurve.AddKey(FocusPoint.Focus, Value);
+			FRichCurveKey& NewKey = OutCurve.GetKey(NewKeyHandle);
+			NewKey.TangentMode = ERichCurveTangentMode::RCTM_None;
+			NewKey.InterpMode = ERichCurveInterpMode::RCIM_Linear;
+		}
+	}
+
+	return true;
+}
+
+void FImageCenterTable::SetParameterCurveKeysAtFocus(float InFocus, int32 InParameterIndex, const FRichCurve& InSourceCurve, TArrayView<const FKeyHandle> InKeys)
+{
+	if (!FParameters::IsValid(InParameterIndex))
+	{
+		return;
+	}
+	
+	if (FImageCenterFocusPoint* FocusPoint = GetFocusPoint(InFocus))
+	{
+		FRichCurve* ActiveCurve = nullptr;
+		if (InParameterIndex == FParameters::Cx)
+		{
+			ActiveCurve = &FocusPoint->Cx;
+		}
+		else
 		{
-			if(ParameterIndex == 0)
+			ActiveCurve = &FocusPoint->Cy;
+		}
+		
+		for (int32 Index = 0; Index < InKeys.Num(); ++Index)
+		{
+			const FKeyHandle Handle = InKeys[Index];
+			const int32 KeyIndex = InSourceCurve.GetIndexSafe(Handle);
+			if (KeyIndex != INDEX_NONE)
 			{
-				OutCurve = FocusPoint->Cx;
+				ActiveCurve->Keys[KeyIndex] = InSourceCurve.GetKey(Handle);
 			}
-			else
+		}
+
+		ActiveCurve->AutoSetTangents();
+	}
+}
+
+void FImageCenterTable::SetParameterCurveKeysAtZoom(float InZoom, int32 InParameterIndex, const FRichCurve& InSourceCurve, TArrayView<const FKeyHandle> InKeys)
+{
+	if (!FParameters::IsValid(InParameterIndex))
+	{
+		return;
+	}
+	
+	for (const FKeyHandle& KeyHandle : InKeys)
+	{
+		// Assume the focus keys are put into the source curve in the same order as they are stored internally
+		const int32 KeyIndex = InSourceCurve.GetIndexSafe(KeyHandle);
+		if (KeyIndex != INDEX_NONE)
+		{
+			if (ensure(FocusPoints.IsValidIndex(KeyIndex)))
 			{
-				OutCurve = FocusPoint->Cy;
+				FImageCenterFocusPoint& FocusPoint = FocusPoints[KeyIndex];
+				FImageCenterInfo ZoomPoint;
+				if (!FocusPoint.GetPoint(InZoom, ZoomPoint))
+				{
+					continue;
+				}
+
+				ZoomPoint.PrincipalPoint[InParameterIndex] = InSourceCurve.GetKeyValue(KeyHandle);
+				FocusPoint.SetPoint(InZoom, ZoomPoint);
+
+				if (InParameterIndex == FParameters::Cx)
+				{
+					FocusPoint.Cx.AutoSetTangents();
+				}
+				else
+				{
+					FocusPoint.Cy.AutoSetTangents();
+				}
 			}
-			return true;
-		}	
+		}
 	}
+}
 
-	return false;
+FText FImageCenterTable::GetParameterValueLabel(int32 InParameterIndex) const
+{
+	return NSLOCTEXT("FImageCenterTable", "ParameterValueLabel", "(normalized)");
 }
 
 const FImageCenterFocusPoint* FImageCenterTable::GetFocusPoint(float InFocus, float InputTolerance) const
diff --git a/Engine/Plugins/VirtualProduction/CameraCalibrationCore/Source/CameraCalibrationCore/Private/Tables/NodalOffsetTable.cpp b/Engine/Plugins/VirtualProduction/CameraCalibrationCore/Source/CameraCalibrationCore/Private/Tables/NodalOffsetTable.cpp
--- a/Engine/Plugins/VirtualProduction/CameraCalibrationCore/Source/CameraCalibrationCore/Private/Tables/NodalOffsetTable.cpp
+++ b/Engine/Plugins/VirtualProduction/CameraCalibrationCore/Source/CameraCalibrationCore/Private/Tables/NodalOffsetTable.cpp
@@ -185,25 +185,190 @@ UScriptStruct* FNodalOffsetTable::GetScriptStruct() const
 	return StaticStruct();
 }
 
-bool FNodalOffsetTable::BuildParameterCurve(float InFocus, int32 ParameterIndex, EAxis::Type InAxis, FRichCurve& OutCurve) const
+bool FNodalOffsetTable::BuildParameterCurveAtFocus(float InFocus, int32 InParameterIndex, FRichCurve& OutCurve) const
 {
-	if((ParameterIndex >= 0) && (ParameterIndex < 2) && (InAxis != EAxis::None))
+	if (!FParameters::IsValidComposed(InParameterIndex))
 	{
-		if(const FNodalOffsetFocusPoint* FocusPoint = GetFocusPoint(InFocus))
+		return false;
+	}
+
+	int32 Parameter;
+	EAxis::Type Axis;
+	FParameters::Decompose(InParameterIndex, Parameter, Axis);
+	
+	if (const FNodalOffsetFocusPoint* FocusPoint = GetFocusPoint(InFocus))
+	{
+		if (Parameter == FParameters::Location)
+		{
+			OutCurve = FocusPoint->LocationOffset[Axis - 1];
+		}
+		else
+		{
+			OutCurve = FocusPoint->RotationOffset[Axis - 1];
+		}
+		
+		return true;
+	}
+
+	return false;
+}
+
+bool FNodalOffsetTable::BuildParameterCurveAtZoom(float InZoom, int32 InParameterIndex, FRichCurve& OutCurve) const
+{
+	if (!FParameters::IsValidComposed(InParameterIndex))
+	{
+		return false;
+	}
+
+	int32 Parameter;
+	EAxis::Type Axis;
+	FParameters::Decompose(InParameterIndex, Parameter, Axis);
+	
+	for (const FNodalOffsetFocusPoint& FocusPoint : FocusPoints)
+	{
+		FNodalPointOffset ZoomPoint;
+		if (FocusPoint.GetPoint(InZoom, ZoomPoint))
+		{
+			const float Value = Parameter == 0 ? ZoomPoint.LocationOffset[Axis - 1] : ZoomPoint.RotationOffset.Rotator().GetComponentForAxis(Axis);
+			const FKeyHandle NewKeyHandle = OutCurve.AddKey(FocusPoint.Focus, Value);
+			FRichCurveKey& NewKey = OutCurve.GetKey(NewKeyHandle);
+			NewKey.TangentMode = ERichCurveTangentMode::RCTM_None;
+			NewKey.InterpMode = ERichCurveInterpMode::RCIM_Linear;
+		}
+	}
+
+	return true;
+}
+
+void FNodalOffsetTable::SetParameterCurveKeysAtFocus(float InFocus, int32 InParameterIndex, const FRichCurve& InSourceCurve, TArrayView<const FKeyHandle> InKeys)
+{
+	if (!FParameters::IsValidComposed(InParameterIndex))
+	{
+		return;
+	}
+	
+	if (FNodalOffsetFocusPoint* FocusPoint = GetFocusPoint(InFocus))
+	{
+		int32 Parameter;
+		EAxis::Type Axis;
+		FParameters::Decompose(InParameterIndex, Parameter, Axis);
+		
+		FRichCurve* ActiveCurve = nullptr;
+		if (Parameter == FParameters::Location)
+		{
+			ActiveCurve = &FocusPoint->LocationOffset[Axis - 1];
+		}
+		else
+		{
+			ActiveCurve = &FocusPoint->RotationOffset[Axis - 1];
+		}
+		
+		for (int32 Index = 0; Index < InKeys.Num(); ++Index)
 		{
-			if(ParameterIndex == 0)
+			const FKeyHandle Handle = InKeys[Index];
+			const int32 KeyIndex = InSourceCurve.GetIndexSafe(Handle);
+			if (KeyIndex != INDEX_NONE)
 			{
-				OutCurve = FocusPoint->LocationOffset[static_cast<uint8>(InAxis) - 1];;
+				ActiveCurve->Keys[KeyIndex] = InSourceCurve.GetKey(Handle);
 			}
-			else
+		}
+
+		ActiveCurve->AutoSetTangents();
+	}
+}
+
+void FNodalOffsetTable::SetParameterCurveKeysAtZoom(float InZoom, int32 InParameterIndex, const FRichCurve& InSourceCurve, TArrayView<const FKeyHandle> InKeys)
+{
+	if (!FParameters::IsValidComposed(InParameterIndex))
+	{
+		return;
+	}
+
+	int32 Parameter;
+	EAxis::Type Axis;
+	FParameters::Decompose(InParameterIndex, Parameter, Axis);
+	
+	for (const FKeyHandle& KeyHandle : InKeys)
+	{
+		// Assume the focus keys are put into the source curve in the same order as they are stored internally
+		const int32 KeyIndex = InSourceCurve.GetIndexSafe(KeyHandle);
+		if (KeyIndex != INDEX_NONE)
+		{
+			if (ensure(FocusPoints.IsValidIndex(KeyIndex)))
 			{
-				OutCurve = FocusPoint->RotationOffset[static_cast<uint8>(InAxis) - 1];
+				FNodalOffsetFocusPoint& FocusPoint = FocusPoints[KeyIndex];
+				FNodalPointOffset ZoomPoint;
+				if (!FocusPoint.GetPoint(InZoom, ZoomPoint))
+				{
+					continue;
+				}
+
+				if (Parameter == FParameters::Location)
+				{
+					ZoomPoint.LocationOffset[Axis - 1] = InSourceCurve.GetKeyValue(KeyHandle);
+				}
+				else if (Parameter == FParameters::Rotation)
+				{
+					FRotator Rotator = ZoomPoint.RotationOffset.Rotator();
+					Rotator.SetComponentForAxis(Axis, InSourceCurve.GetKeyValue(KeyHandle));
+					ZoomPoint.RotationOffset = Rotator.Quaternion();
+				}
+
+				FocusPoint.SetPoint(InZoom, ZoomPoint);
+
+				if (Parameter == FParameters::Location)
+				{
+					FocusPoint.LocationOffset[Axis - 1].AutoSetTangents();
+				}
+				else if (Parameter == FParameters::Rotation)
+				{
+					FocusPoint.RotationOffset[Axis - 1].AutoSetTangents();
+				}
 			}
-			return true;
-		}	
+		}
 	}
+}
 
-	return false;
+FText FNodalOffsetTable::GetParameterValueLabel(int32 InParameterIndex) const
+{
+	if (!FParameters::IsValidComposed(InParameterIndex))
+	{
+		return FText();
+	}
+	
+	int32 Parameter;
+	EAxis::Type Axis;
+	FParameters::Decompose(InParameterIndex, Parameter, Axis);
+	
+	if (Parameter == FParameters::Location)
+	{
+		return NSLOCTEXT("FNodalOffsetTable", "LocationParameterValueLabel", "(cm)");
+	}
+	else
+	{
+		return NSLOCTEXT("FNodalOffsetTable", "RotationParameterValueLabel", "(deg)");
+	}
+}
+
+FText FNodalOffsetTable::GetParameterValueUnitLabel(int32 InParameterIndex) const
+{
+	if (!FParameters::IsValidComposed(InParameterIndex))
+	{
+		return FText();
+	}
+	
+	int32 Parameter;
+	EAxis::Type Axis;
+	FParameters::Decompose(InParameterIndex, Parameter, Axis);
+	
+	if (Parameter == FParameters::Location)
+	{
+		return NSLOCTEXT("FNodalOffsetTable", "LocationParameterUnitLabel", "cm");
+	}
+	else
+	{
+		return NSLOCTEXT("FNodalOffsetTable", "RotationParameterUnitLabel", "deg");
+	}
 }
 
 const FNodalOffsetFocusPoint* FNodalOffsetTable::GetFocusPoint(float InFocus, float InputTolerance) const
diff --git a/Engine/Plugins/VirtualProduction/CameraCalibrationCore/Source/CameraCalibrationCore/Private/Tables/STMapTable.cpp b/Engine/Plugins/VirtualProduction/CameraCalibrationCore/Source/CameraCalibrationCore/Private/Tables/STMapTable.cpp
--- a/Engine/Plugins/VirtualProduction/CameraCalibrationCore/Source/CameraCalibrationCore/Private/Tables/STMapTable.cpp
+++ b/Engine/Plugins/VirtualProduction/CameraCalibrationCore/Source/CameraCalibrationCore/Private/Tables/STMapTable.cpp
@@ -174,9 +174,9 @@ UScriptStruct* FSTMapTable::GetScriptStruct() const
 	return StaticStruct();
 }
 
-bool FSTMapTable::BuildMapBlendingCurve(float InFocus, FRichCurve& OutCurve)
+bool FSTMapTable::BuildParameterCurveAtFocus(float InFocus, int32 InParameterIndex, FRichCurve& OutCurve) const
 {
-	if(FSTMapFocusPoint* FocusPoint = GetFocusPoint(InFocus))
+	if (const FSTMapFocusPoint* FocusPoint = GetFocusPoint(InFocus))
 	{
 		OutCurve = FocusPoint->MapBlendingCurve;
 		return true;
@@ -185,6 +185,64 @@ bool FSTMapTable::BuildMapBlendingCurve(float InFocus, FRichCurve& OutCurve)
 	return false;
 }
 
+bool FSTMapTable::BuildParameterCurveAtZoom(float InZoom, int32 InParameterIndex, FRichCurve& OutCurve) const
+{
+	for (const FSTMapFocusPoint& FocusPoint : FocusPoints)
+	{
+		FKeyHandle ZoomKeyHandle = FocusPoint.MapBlendingCurve.FindKey(InZoom);
+		if (FocusPoint.MapBlendingCurve.IsKeyHandleValid(ZoomKeyHandle))
+		{
+			const float Value = FocusPoint.MapBlendingCurve.GetKeyValue(ZoomKeyHandle);
+			const FKeyHandle NewKeyHandle = OutCurve.AddKey(FocusPoint.Focus, Value);
+			FRichCurveKey& NewKey = OutCurve.GetKey(NewKeyHandle);
+			NewKey.TangentMode = ERichCurveTangentMode::RCTM_None;
+			NewKey.InterpMode = ERichCurveInterpMode::RCIM_Linear;
+		}
+	}
+
+	return true;
+}
+
+void FSTMapTable::SetParameterCurveKeysAtFocus(float InFocus, int32 InParameterIndex, const FRichCurve& InSourceCurve, TArrayView<const FKeyHandle> InKeys)
+{
+	if (FSTMapFocusPoint* FocusPoint = GetFocusPoint(InFocus))
+	{
+		for (int32 Index = 0; Index < InKeys.Num(); ++Index)
+		{
+			const FKeyHandle Handle = InKeys[Index];
+			const int32 KeyIndex = InSourceCurve.GetIndexSafe(Handle);
+			if (KeyIndex != INDEX_NONE)
+			{
+				//We can't move keys on the time axis so our indices should match
+				const FRichCurveKey& Key = InSourceCurve.GetKey(Handle);
+				FocusPoint->MapBlendingCurve.Keys[KeyIndex] = Key;
+			}
+		}
+
+		FocusPoint->MapBlendingCurve.AutoSetTangents();
+	}
+}
+
+void FSTMapTable::SetParameterCurveKeysAtZoom(float InZoom, int32 InParameterIndex, const FRichCurve& InSourceCurve, TArrayView<const FKeyHandle> InKeys)
+{
+	for (const FKeyHandle& KeyHandle : InKeys)
+	{
+		// Assume the focus keys are put into the source curve in the same order as they are stored internally
+		const int32 KeyIndex = InSourceCurve.GetIndexSafe(KeyHandle);
+		if (KeyIndex != INDEX_NONE)
+		{
+			if (ensure(FocusPoints.IsValidIndex(KeyIndex)))
+			{
+				FSTMapFocusPoint& FocusPoint = FocusPoints[KeyIndex];
+				
+				//We can't move keys on the time axis so our indices should match
+				FocusPoint.MapBlendingCurve.Keys[KeyIndex].Value = InSourceCurve.GetKeyValue(KeyHandle);
+				FocusPoint.MapBlendingCurve.AutoSetTangents();
+			}
+		}
+	}
+}
+
 const FSTMapFocusPoint* FSTMapTable::GetFocusPoint(float InFocus, float InputTolerance) const
 {
 	return FocusPoints.FindByPredicate([InFocus, InputTolerance](const FSTMapFocusPoint& Point) { return FMath::IsNearlyEqual(Point.Focus, InFocus, InputTolerance); });
diff --git a/Engine/Plugins/VirtualProduction/CameraCalibrationCore/Source/CameraCalibrationCore/Public/LensFile.h b/Engine/Plugins/VirtualProduction/CameraCalibrationCore/Source/CameraCalibrationCore/Public/LensFile.h
--- a/Engine/Plugins/VirtualProduction/CameraCalibrationCore/Source/CameraCalibrationCore/Public/LensFile.h
+++ b/Engine/Plugins/VirtualProduction/CameraCalibrationCore/Source/CameraCalibrationCore/Public/LensFile.h
@@ -232,6 +232,9 @@ public:
 
 	/** Get data table reference based on given category */
 	const FBaseLensTable* GetDataTable(ELensDataCategory InDataCategory) const;
+
+	/** Get data table reference based on given category */
+	FBaseLensTable* GetDataTable(ELensDataCategory InDataCategory);
 	
 	/** Returns the delegate that is triggered when the LensModel changes */
 	FOnLensFileModelChanged& OnLensFileModelChanged() { return OnLensFileModelChangedDelegate; }
diff --git a/Engine/Plugins/VirtualProduction/CameraCalibrationCore/Source/CameraCalibrationCore/Public/Tables/BaseLensTable.h b/Engine/Plugins/VirtualProduction/CameraCalibrationCore/Source/CameraCalibrationCore/Public/Tables/BaseLensTable.h
--- a/Engine/Plugins/VirtualProduction/CameraCalibrationCore/Source/CameraCalibrationCore/Public/Tables/BaseLensTable.h
+++ b/Engine/Plugins/VirtualProduction/CameraCalibrationCore/Source/CameraCalibrationCore/Public/Tables/BaseLensTable.h
@@ -8,6 +8,8 @@
 
 #include "BaseLensTable.generated.h"
 
+struct FKeyHandle;
+struct FRichCurve;
 class ULensFile;
 enum class ELensDataCategory : uint8;
 
@@ -114,6 +116,66 @@ public:
 	/** Get Struct class of this Data Table */
 	virtual UScriptStruct* GetScriptStruct() const PURE_VIRTUAL(FBaseLensTable::GetFocusPointNum, return nullptr; );
 
+	/** 
+	* Fills OutCurve with all points contained in the given focus 
+	* Returns false if FocusIdentifier is not found or ParameterIndex isn't valid
+	*/
+	virtual bool BuildParameterCurveAtFocus(float InFocus, int32 ParameterIndex, FRichCurve& OutCurve) const PURE_VIRTUAL(FBaseLensTable::BuildParameterCurve, return false; ); 
+
+	/**
+	 * Fills OutCurve with points across all focuses that have the given zoom
+	 * @param InZoom The zoom to get the curve for
+	 * @param ParameterIndex The index of the data parameter to get the curve for
+	 * @param OutCurve The curve to fill
+	 * @return true if the curve was built, otherwise false
+	 */
+	virtual bool BuildParameterCurveAtZoom(float InZoom, int32 ParameterIndex, FRichCurve& OutCurve) const PURE_VIRTUAL(FBaseLensTable::BuildFocusParameterCurve, return false; );
+
+	/**
+	 * Updates the keys of all zoom points at the specified focus to match the corresponding keys in the specified curve
+	 * @param InFocus The focus whose zoom points should be updated
+	 * @param InParameterIndex The index of the parameter whose curve keys are being updated
+	 * @param InSourceCurve The curve to copy the keys from
+	 * @param InKeys The keys to copy from the source curve into the zoom parameter curve
+	 */
+	virtual void SetParameterCurveKeysAtFocus(float InFocus, int32 InParameterIndex, const FRichCurve& InSourceCurve, TArrayView<const FKeyHandle> InKeys) PURE_VIRTUAL(FBaseLensTable::SetZoomParameterCurveKeys);
+
+	/**
+	 * Updates the keys of all focus points at the specified zoom to match the corresponding keys in the specified curve
+	 * @param InZoom The zoom whose focus points should be updated
+	 * @param InParameterIndex The index of the parameter whose curve keys are being updated
+	 * @param InSourceCurve The curve to copy the keys from
+	 * @param InKeys The keys to copy from the source curve into the focus parameter curve
+	 */
+	virtual void SetParameterCurveKeysAtZoom(float InZoom, int32 InParameterIndex, const FRichCurve& InSourceCurve, TArrayView<const FKeyHandle> InKeys) PURE_VIRTUAL(FBaseLensTable::SetFocusParameterCurveKeys);
+
+	/**
+	 * Gets whether the positions of the table's curve keys can be edited or not
+	 * @param InParameterIndex The index of the parameter of the curve in question
+	 * @return true if the curve's key positions can be edited, false otherwise
+	 */
+	virtual bool CanEditCurveKeyPositions(int32 InParameterIndex) const PURE_VIRTUAL(FBaseLensTable::CanEditCurveKeyPositions, return false; );
+
+	/**
+	 * Gets whether the attributes of the table's curve keys can be edited or not
+	 * @param InParameterIndex The index of the parameter of the curve in question
+	 * @return true if the curve's key attributes can be edited, false otherwise
+	 */
+	virtual bool CanEditCurveKeyAttributes(int32 InParameterIndex) const PURE_VIRTUAL(FBaseLensTable::CanEditCurveKeyAttributes, return false; );
+
+	/**
+	 * Gets the range of allowed values for the curve keys at the specified parameter index
+	 * @param InParameterIndex The index of the parameter to get the ranges for
+	 * @return The minimum and maximum allowed values
+	 */
+	virtual TRange<double> GetCurveKeyPositionRange(int32 InParameterIndex) const { return TRange<double>(TNumericLimits<double>::Lowest(), TNumericLimits<double>::Max()); }
+	
+	/** Gets the text to display on any UI when labeling the values of this table's parameters */
+	virtual FText GetParameterValueLabel(int32 InParameterIndex) const { return FText(); }
+
+	/** Gets the text to display on any UI when displaying units of the values of this table's parameters */
+	virtual FText GetParameterValueUnitLabel(int32 InParameterIndex) const { return FText(); }
+	
 	/** Get Names of this Data Point */
 	static FName GetFriendlyPointName(ELensDataCategory InCategory);
 
diff --git a/Engine/Plugins/VirtualProduction/CameraCalibrationCore/Source/CameraCalibrationCore/Public/Tables/DistortionParametersTable.h b/Engine/Plugins/VirtualProduction/CameraCalibrationCore/Source/CameraCalibrationCore/Public/Tables/DistortionParametersTable.h
--- a/Engine/Plugins/VirtualProduction/CameraCalibrationCore/Source/CameraCalibrationCore/Public/Tables/DistortionParametersTable.h
+++ b/Engine/Plugins/VirtualProduction/CameraCalibrationCore/Source/CameraCalibrationCore/Public/Tables/DistortionParametersTable.h
@@ -89,6 +89,12 @@ struct CAMERACALIBRATIONCORE_API FDistortionTable : public FBaseLensTable
 
 	using FocusPointType = FDistortionFocusPoint;
 
+	/** Wrapper for indices of specific parameters for the distortion table  */
+	struct FParameters
+	{
+		static constexpr int32 Aggregate = INDEX_NONE;
+	};
+	
 protected:
 	//~ Begin FBaseDataTable Interface
 	virtual TMap<ELensDataCategory, FLinkPointMetadata> GetLinkedCategories() const override;
@@ -103,14 +109,15 @@ public:
 	virtual int32 GetFocusPointNum() const override { return FocusPoints.Num(); }
 	virtual int32 GetTotalPointNum() const override;
 	virtual UScriptStruct* GetScriptStruct() const override;
+	virtual bool BuildParameterCurveAtFocus(float InFocus, int32 InParameterIndex, FRichCurve& OutCurve) const override;
+	virtual bool BuildParameterCurveAtZoom(float InZoom, int32 InParameterIndex, FRichCurve& OutCurve) const override;
+	virtual void SetParameterCurveKeysAtFocus(float InFocus, int32 InParameterIndex, const FRichCurve& InSourceCurve, TArrayView<const FKeyHandle> InKeys) override;
+	virtual void SetParameterCurveKeysAtZoom(float InZoom, int32 InParameterIndex, const FRichCurve& InSourceCurve, TArrayView<const FKeyHandle> InKeys) override;
+	virtual bool CanEditCurveKeyPositions(int32 InParameterIndex) const override;
+	virtual bool CanEditCurveKeyAttributes(int32 InParameterIndex) const override;
+	virtual FText GetParameterValueLabel(int32 InParameterIndex) const override;
 	//~ End FBaseDataTable Interface
 
-	/** 
-	 * Fills OutCurve with all points contained in the given focus 
-	 * Returns false if FocusIdentifier is not found or ParameterIndex isn't valid
-	 */
-	bool BuildParameterCurve(float InFocus, int32 ParameterIndex, FRichCurve& OutCurve) const;
-
 	/** Returns const point for a given focus */
 	const FDistortionFocusPoint* GetFocusPoint(float InFocus, float InputTolerance = KINDA_SMALL_NUMBER) const;
 	
diff --git a/Engine/Plugins/VirtualProduction/CameraCalibrationCore/Source/CameraCalibrationCore/Public/Tables/FocalLengthTable.h b/Engine/Plugins/VirtualProduction/CameraCalibrationCore/Source/CameraCalibrationCore/Public/Tables/FocalLengthTable.h
--- a/Engine/Plugins/VirtualProduction/CameraCalibrationCore/Source/CameraCalibrationCore/Public/Tables/FocalLengthTable.h
+++ b/Engine/Plugins/VirtualProduction/CameraCalibrationCore/Source/CameraCalibrationCore/Public/Tables/FocalLengthTable.h
@@ -99,6 +99,20 @@ struct CAMERACALIBRATIONCORE_API FFocalLengthTable : public FBaseLensTable
 
 	using FocusPointType = FFocalLengthFocusPoint;
 
+	/** Wrapper for indices of specific parameters for the focal length table  */
+	struct FParameters
+	{
+		static constexpr int32 Aggregate = INDEX_NONE;
+		static constexpr int32 Fx = 0;
+		static constexpr int32 Fy = 1;
+
+		/** Returns if a parameter index is valid (not including the aggregate value) */
+		static bool IsValid(int32 InParameterIndex) { return InParameterIndex >= 0 && InParameterIndex < 2; }
+		
+		/** Returns if a parameter index is valid or the aggregate value */
+		static bool IsValidOrAggregate(int32 InParameterIndex) { return IsValid(InParameterIndex) || InParameterIndex == Aggregate; }
+	};
+	
 protected:
 	//~ Begin FBaseDataTable Interface
 	virtual TMap<ELensDataCategory, FLinkPointMetadata> GetLinkedCategories() const override;
@@ -113,13 +127,16 @@ public:
 	virtual int32 GetFocusPointNum() const override { return FocusPoints.Num(); }
 	virtual int32 GetTotalPointNum() const override;
 	virtual UScriptStruct* GetScriptStruct() const override;
+	virtual bool BuildParameterCurveAtFocus(float InFocus, int32 ParameterIndex, FRichCurve& OutCurve) const override;
+	virtual bool BuildParameterCurveAtZoom(float InZoom, int32 ParameterIndex, FRichCurve& OutCurve) const override;
+	virtual void SetParameterCurveKeysAtFocus(float InFocus, int32 InParameterIndex, const FRichCurve& InSourceCurve, TArrayView<const FKeyHandle> InKeys) override;
+	virtual void SetParameterCurveKeysAtZoom(float InZoom, int32 InParameterIndex, const FRichCurve& InSourceCurve, TArrayView<const FKeyHandle> InKeys) override;
+	virtual bool CanEditCurveKeyPositions(int32 InParameterIndex) const override { return true; }
+	virtual bool CanEditCurveKeyAttributes(int32 InParameterIndex) const override { return true; }
+	virtual TRange<double> GetCurveKeyPositionRange(int32 InParameterIndex) const override;
+	virtual FText GetParameterValueLabel(int32 InParameterIndex) const override;
+	virtual FText GetParameterValueUnitLabel(int32 InParameterIndex) const override;
 	//~ End FBaseDataTable Interface
-
-	/** 
-	* Fills OutCurve with all points contained in the given focus 
-	* Returns false if FocusIdentifier is not found or ParameterIndex isn't valid
-	*/
-	bool BuildParameterCurve(float InFocus, int32 ParameterIndex, FRichCurve& OutCurve) const;
 	
 	/** Returns const point for a given focus */
 	const FFocalLengthFocusPoint* GetFocusPoint(float InFocus, float InputTolerance = KINDA_SMALL_NUMBER) const;
diff --git a/Engine/Plugins/VirtualProduction/CameraCalibrationCore/Source/CameraCalibrationCore/Public/Tables/ImageCenterTable.h b/Engine/Plugins/VirtualProduction/CameraCalibrationCore/Source/CameraCalibrationCore/Public/Tables/ImageCenterTable.h
--- a/Engine/Plugins/VirtualProduction/CameraCalibrationCore/Source/CameraCalibrationCore/Public/Tables/ImageCenterTable.h
+++ b/Engine/Plugins/VirtualProduction/CameraCalibrationCore/Source/CameraCalibrationCore/Public/Tables/ImageCenterTable.h
@@ -69,6 +69,16 @@ struct CAMERACALIBRATIONCORE_API FImageCenterTable : public FBaseLensTable
 
 	using FocusPointType = FImageCenterFocusPoint;
 
+	/** Wrapper for indices of specific parameters for the image center table  */
+	struct FParameters
+	{
+		static constexpr int32 Cx = 0;
+		static constexpr int32 Cy = 1;
+
+		/** Returns if a parameter index is valid */
+		static bool IsValid(int32 InParameterIndex) { return InParameterIndex >= 0 && InParameterIndex < 2; }
+	};
+	
 protected:
 	//~ Begin FBaseDataTable Interface
 	virtual TMap<ELensDataCategory, FLinkPointMetadata> GetLinkedCategories() const override;
@@ -83,14 +93,15 @@ public:
 	virtual int32 GetFocusPointNum() const override { return FocusPoints.Num(); }
 	virtual int32 GetTotalPointNum() const override;
 	virtual UScriptStruct* GetScriptStruct() const override;
+	virtual bool BuildParameterCurveAtFocus(float InFocus, int32 ParameterIndex, FRichCurve& OutCurve) const override;
+	virtual bool BuildParameterCurveAtZoom(float InZoom, int32 ParameterIndex, FRichCurve& OutCurve) const override;
+	virtual void SetParameterCurveKeysAtFocus(float InFocus, int32 InParameterIndex, const FRichCurve& InSourceCurve, TArrayView<const FKeyHandle> InKeys) override;
+	virtual void SetParameterCurveKeysAtZoom(float InZoom, int32 InParameterIndex, const FRichCurve& InSourceCurve, TArrayView<const FKeyHandle> InKeys) override;
+	virtual bool CanEditCurveKeyPositions(int32 InParameterIndex) const override { return true; }
+	virtual bool CanEditCurveKeyAttributes(int32 InParameterIndex) const override { return true; }
+	virtual FText GetParameterValueLabel(int32 InParameterIndex) const override;
 	//~ End FBaseDataTable Interface
 	
-	/** 
-	* Fills OutCurve with all points contained in the given focus 
-	* Returns false if FocusIdentifier is not found or ParameterIndex isn't valid
-	*/
-	bool BuildParameterCurve(float InFocus, int32 ParameterIndex, FRichCurve& OutCurve) const;
-	
 	/** Returns const point for a given focus */
 	const FImageCenterFocusPoint* GetFocusPoint(float InFocus, float InputTolerance = KINDA_SMALL_NUMBER) const;
 
diff --git a/Engine/Plugins/VirtualProduction/CameraCalibrationCore/Source/CameraCalibrationCore/Public/Tables/NodalOffsetTable.h b/Engine/Plugins/VirtualProduction/CameraCalibrationCore/Source/CameraCalibrationCore/Public/Tables/NodalOffsetTable.h
--- a/Engine/Plugins/VirtualProduction/CameraCalibrationCore/Source/CameraCalibrationCore/Public/Tables/NodalOffsetTable.h
+++ b/Engine/Plugins/VirtualProduction/CameraCalibrationCore/Source/CameraCalibrationCore/Public/Tables/NodalOffsetTable.h
@@ -75,6 +75,26 @@ struct CAMERACALIBRATIONCORE_API FNodalOffsetTable : public FBaseLensTable
 
 	using FocusPointType = FNodalOffsetFocusPoint;
 
+	/** Wrapper for indices of specific parameters for the nodal offset table  */
+	struct FParameters
+	{
+		static constexpr int32 Location = 0;
+		static constexpr int32 Rotation = 1;
+
+		/** Composes the parameter and axis indices into a single value */
+		static int32 Compose(int32 InParameterIndex, EAxis::Type Axis) { return InParameterIndex * 3 + (Axis - 1); }
+
+		/** Composes a combined index into a parameter index and an axis */
+		static void Decompose(int32 InComposedIndex, int32& OutParameterIndex, EAxis::Type& OutAxis)
+		{
+			OutParameterIndex = InComposedIndex / 3;
+			OutAxis = (EAxis::Type)(InComposedIndex % 3 + 1);
+		}
+		
+		/** Returns if a composed parameter index is valid */
+		static bool IsValidComposed(int32 InComposedIndex) { return InComposedIndex >= 0 && InComposedIndex < 6; }
+	};
+	
 protected:
 	//~ Begin FBaseDataTable Interface
 	virtual TMap<ELensDataCategory, FLinkPointMetadata> GetLinkedCategories() const override;
@@ -89,14 +109,16 @@ public:
 	virtual int32 GetFocusPointNum() const override { return FocusPoints.Num(); }
 	virtual int32 GetTotalPointNum() const override;
 	virtual UScriptStruct* GetScriptStruct() const override;
+	virtual bool BuildParameterCurveAtFocus(float InFocus, int32 InParameterIndex, FRichCurve& OutCurve) const override;
+	virtual bool BuildParameterCurveAtZoom(float InZoom, int32 InParameterIndex, FRichCurve& OutCurve) const override;
+	virtual void SetParameterCurveKeysAtFocus(float InFocus, int32 InParameterIndex, const FRichCurve& InSourceCurve, TArrayView<const FKeyHandle> InKeys) override;
+	virtual void SetParameterCurveKeysAtZoom(float InZoom, int32 InParameterIndex, const FRichCurve& InSourceCurve, TArrayView<const FKeyHandle> InKeys) override;
+	virtual bool CanEditCurveKeyPositions(int32 InParameterIndex) const override { return true; }
+	virtual bool CanEditCurveKeyAttributes(int32 InParameterIndex) const override { return true; }
+	virtual FText GetParameterValueLabel(int32 InParameterIndex) const override;
+	virtual FText GetParameterValueUnitLabel(int32 InParameterIndex) const override;
 	//~ End FBaseDataTable Interface
 
-	/** 
-	* Fills OutCurve with all points contained in the given focus 
-	* Returns false if FocusIdentifier is not found or ParameterIndex isn't valid
-	*/
-	bool BuildParameterCurve(float InFocus, int32 ParameterIndex, EAxis::Type InAxis, FRichCurve& OutCurve) const;
-
 	/** Returns const point for a given focus */
 	const FNodalOffsetFocusPoint* GetFocusPoint(float InFocus, float InputTolerance = KINDA_SMALL_NUMBER) const;
 
diff --git a/Engine/Plugins/VirtualProduction/CameraCalibrationCore/Source/CameraCalibrationCore/Public/Tables/STMapTable.h b/Engine/Plugins/VirtualProduction/CameraCalibrationCore/Source/CameraCalibrationCore/Public/Tables/STMapTable.h
--- a/Engine/Plugins/VirtualProduction/CameraCalibrationCore/Source/CameraCalibrationCore/Public/Tables/STMapTable.h
+++ b/Engine/Plugins/VirtualProduction/CameraCalibrationCore/Source/CameraCalibrationCore/Public/Tables/STMapTable.h
@@ -141,14 +141,14 @@ public:
 	virtual int32 GetFocusPointNum() const override { return FocusPoints.Num(); }
 	virtual int32 GetTotalPointNum() const override;
 	virtual UScriptStruct* GetScriptStruct() const override;
+	virtual bool BuildParameterCurveAtFocus(float InFocus, int32 ParameterIndex, FRichCurve& OutCurve) const override;
+	virtual bool BuildParameterCurveAtZoom(float InZoom, int32 ParameterIndex, FRichCurve& OutCurve) const override;
+	virtual void SetParameterCurveKeysAtFocus(float InFocus, int32 InParameterIndex, const FRichCurve& InSourceCurve, TArrayView<const FKeyHandle> InKeys) override;
+	virtual void SetParameterCurveKeysAtZoom(float InZoom, int32 InParameterIndex, const FRichCurve& InSourceCurve, TArrayView<const FKeyHandle> InKeys) override;
+	virtual bool CanEditCurveKeyPositions(int32 InParameterIndex) const override { return false; }
+	virtual bool CanEditCurveKeyAttributes(int32 InParameterIndex) const override { return true; }
 	//~ End FBaseDataTable Interface
 
-	/** 
-	 * Builds the map blending curve into OutCurve
-	 * Returns true if focus point exists
-	 */
-	bool BuildMapBlendingCurve(float InFocus, FRichCurve& OutCurve);
-	
 	/** Returns const point for a given focus */
 	const FSTMapFocusPoint* GetFocusPoint(float InFocus, float InputTolerance = KINDA_SMALL_NUMBER) const;
 
