CameraCalibration: Add utility to group point sets by the same camera pose before sending to the distortion solver. This effectively constrains the solver to find a single camera pose that solves a larger set of points (or a single set of points in the case of a stationary camera).
Add support for the distortion solver to compute the nodal offset if it solved for exactly one camera pose, and a feature for the lens distortion tool to add the nodal offset result to the LensFile at the same time that it adds the distortion and intrinsics results.
Update the lens distortion "points" algorithm to use the async calibration task pattern.
Update the FixZeroDistortion solver flag to just FixDistortion, and allow for input distortion parameters to be used to initialize the solver.

diff --git a/Engine/Plugins/VirtualProduction/CameraCalibration/Source/CameraCalibrationEditor/Private/AssetEditor/LensDistortionTool.cpp b/Engine/Plugins/VirtualProduction/CameraCalibration/Source/CameraCalibrationEditor/Private/AssetEditor/LensDistortionTool.cpp
--- a/Engine/Plugins/VirtualProduction/CameraCalibration/Source/CameraCalibrationEditor/Private/AssetEditor/LensDistortionTool.cpp
+++ b/Engine/Plugins/VirtualProduction/CameraCalibration/Source/CameraCalibrationEditor/Private/AssetEditor/LensDistortionTool.cpp
@@ -29,6 +29,10 @@
 
 #define LOCTEXT_NAMESPACE "LensDistortionTool"
 
+#if WITH_EDITOR
+static TAutoConsoleVariable<bool> CVarSolveNodalOffset(TEXT("LensDistortionTool.SolveNodalOffset"), false, TEXT("If true, and if the calibration result contained a valid nodal offset, that offset will be added to the LensFile in addition to distortion."));
+#endif
+
 namespace UE::CameraCalibration::Private::LensDistortionTool
 {
 	static const FString SessionDateTimeField(TEXT("SessionDateTime"));
@@ -503,7 +507,7 @@ void ULensDistortionTool::SaveCalibrationResult()
 			UE_LOG(LogCameraCalibrationEditor, Log, TEXT("The LensFile's data mode was set to ST Map, but the latest calibration result returned distortion parameters. Data mode will change to Parameters."));
 		}
 
-		FScopedTransaction Transaction(LOCTEXT("SaveCurrentDistortionData", "Save Current Distortion Data"));
+		FScopedTransaction Transaction(LOCTEXT("SaveCurrentDistortionData", "Save Calibrated Distortion to Lens Asset"));
 		LensFile->Modify();
 
 		LensFile->AddDistortionPoint(CalibrationResult.EvaluatedFocus, CalibrationResult.EvaluatedZoom, CalibrationResult.Parameters, CalibrationResult.FocalLength);
@@ -517,7 +521,7 @@ void ULensDistortionTool::SaveCalibrationResult()
 			UE_LOG(LogCameraCalibrationEditor, Log, TEXT("The LensFile's data mode was set to Parameters, but the latest calibration result returned an ST Map. Data mode will change to ST Map."));
 		}
 
-		FScopedTransaction Transaction(LOCTEXT("SaveCurrentDistortionData", "Save Current Distortion Data"));
+		FScopedTransaction Transaction(LOCTEXT("SaveCurrentDistortionData", "Save Calibrated Distortion to Lens Asset"));
 		LensFile->Modify();
 
 		LensFile->AddSTMapPoint(CalibrationResult.EvaluatedFocus, CalibrationResult.EvaluatedZoom, CalibrationResult.STMap);
@@ -525,6 +529,18 @@ void ULensDistortionTool::SaveCalibrationResult()
 		LensFile->AddImageCenterPoint(CalibrationResult.EvaluatedFocus, CalibrationResult.EvaluatedZoom, CalibrationResult.ImageCenter);
 	}
 
+	// If the calibration result contains a meaningful nodal offset result, add it to the Lens File
+	if (CVarSolveNodalOffset.GetValueOnGameThread())
+	{
+		if (!CalibrationResult.NodalOffset.LocationOffset.Equals(FVector::ZeroVector) || !CalibrationResult.NodalOffset.RotationOffset.Equals(FQuat::Identity))
+		{
+			FScopedTransaction Transaction(LOCTEXT("SaveNodalOffsetResult", "Save Calibrated Nodal Offset to Lens Asset"));
+			LensFile->Modify();
+
+			LensFile->AddNodalOffsetPoint(CalibrationResult.EvaluatedFocus, CalibrationResult.EvaluatedFocus, CalibrationResult.NodalOffset);
+		}
+	}
+
 	if (UCameraLensDistortionAlgo* Algo = GetAlgo())
 	{
 		Algo->OnDistortionSavedToLens();
diff --git a/Engine/Plugins/VirtualProduction/CameraCalibration/Source/CameraCalibrationEditor/Private/Calibrators/CameraCalibrationSolver.cpp b/Engine/Plugins/VirtualProduction/CameraCalibration/Source/CameraCalibrationEditor/Private/Calibrators/CameraCalibrationSolver.cpp
--- a/Engine/Plugins/VirtualProduction/CameraCalibration/Source/CameraCalibrationEditor/Private/Calibrators/CameraCalibrationSolver.cpp
+++ b/Engine/Plugins/VirtualProduction/CameraCalibration/Source/CameraCalibrationEditor/Private/Calibrators/CameraCalibrationSolver.cpp
@@ -52,6 +52,7 @@ FDistortionCalibrationResult ULensDistortionSolverOpenCV::Solve_Implementation(
 	const FIntPoint ImageSize,
 	const FVector2D& FocalLength,
 	const FVector2D& ImageCenter,
+	const TArray<float>& DistortionParameters,
 	const TArray<FTransform>& CameraPoses,
 	TSubclassOf<ULensModel> LensModel,
 	double PixelAspect,
@@ -95,18 +96,40 @@ FDistortionCalibrationResult ULensDistortionSolverOpenCV::Solve_Implementation(
 	const int NumDistortionCoefficients = LensModel->GetDefaultObject<ULensModel>()->GetNumParameters();
 	cv::Mat DistCoeffs = cv::Mat(1, NumDistortionCoefficients, CV_64F);
 
-	TArray<float> DefaultCoefficients;
-	LensModel->GetDefaultObject<ULensModel>()->GetDefaultParameterArray(DefaultCoefficients);
-
-	for (int CoeffIndex = 0; CoeffIndex < NumDistortionCoefficients; ++CoeffIndex)
+	// If input parameters were provided, they can be used to initialize the solver. Otherwise, the solver will be initialized with the default values for the input model.
+	if (EnumHasAnyFlags(SolverFlags, ECalibrationFlags::UseIntrinsicGuess) && (DistortionParameters.Num() == NumDistortionCoefficients))
 	{
-		DistCoeffs.at<double>(CoeffIndex) = DefaultCoefficients[CoeffIndex];
+		for (int CoeffIndex = 0; CoeffIndex < NumDistortionCoefficients; ++CoeffIndex)
+		{
+			DistCoeffs.at<double>(CoeffIndex) = DistortionParameters[CoeffIndex];
+		}
+	}
+	else
+	{
+		TArray<float> DefaultCoefficients;
+		LensModel->GetDefaultObject<ULensModel>()->GetDefaultParameterArray(DefaultCoefficients);
+
+		for (int CoeffIndex = 0; CoeffIndex < NumDistortionCoefficients; ++CoeffIndex)
+		{
+			DistCoeffs.at<double>(CoeffIndex) = DefaultCoefficients[CoeffIndex];
+		}
 	}
 
 	if (LensModel == UAnamorphicLensModel::StaticClass())
 	{
 		DistCoeffs.at<double>(0) = PixelAspect;
 	}
+	else if (LensModel == USphericalLensModel::StaticClass())
+	{
+		// The spherical distortion coefficients in our model are in a different order than they appear in the solver
+		// so the initial values need to be rearranged before initializing the solver
+		const double K3 = DistCoeffs.at<double>(2);
+		const double P1 = DistCoeffs.at<double>(3);
+		const double P2 = DistCoeffs.at<double>(4);
+		DistCoeffs.at<double>(2) = P1;
+		DistCoeffs.at<double>(3) = P2;
+		DistCoeffs.at<double>(4) = K3;
+	}
 
 	cv::Size CvImageSize = cv::Size(ImageSize.X, ImageSize.Y);
 
@@ -272,7 +295,7 @@ FDistortionCalibrationResult ULensDistortionSolverOpenCV::Solve_Implementation(
 			Mask[ParamIndex] = 0;
 		}
 	}
-	if (EnumHasAnyFlags(SolverFlags, ECalibrationFlags::FixZeroDistortion))
+	if (EnumHasAnyFlags(SolverFlags, ECalibrationFlags::FixDistortion))
 	{
 		for (int ParamIndex = 0; ParamIndex < NumDistortionCoefficients; ++ParamIndex)
 		{
@@ -477,6 +500,17 @@ FDistortionCalibrationResult ULensDistortionSolverOpenCV::Solve_Implementation(
 		Result.CameraPoses.Add(CameraPoseForImage);
 	}
 
+	// If the solver computed exactly one calibrated camera pose, we can compute a tracking offset based on the input camera pose
+	if (Result.CameraPoses.Num() == 1 && CameraPoses.Num() == 1)
+	{
+		const FTransform& TrackedCameraPose = CameraPoses[0];
+		const FTransform& CalibratedCameraPose = Result.CameraPoses[0];
+		const FTransform TrackingOffset = CalibratedCameraPose * TrackedCameraPose.Inverse();
+
+		Result.NodalOffset.LocationOffset = TrackingOffset.GetLocation();
+		Result.NodalOffset.RotationOffset = TrackingOffset.GetRotation();
+	}
+
 	if (LensModel == UAnamorphicLensModel::StaticClass())
 	{
 		for (int CoeffIndex = 0; CoeffIndex < NumDistortionCoefficients; ++CoeffIndex)
diff --git a/Engine/Plugins/VirtualProduction/CameraCalibration/Source/CameraCalibrationEditor/Private/Calibrators/CameraCalibrationSolver.h b/Engine/Plugins/VirtualProduction/CameraCalibration/Source/CameraCalibrationEditor/Private/Calibrators/CameraCalibrationSolver.h
--- a/Engine/Plugins/VirtualProduction/CameraCalibration/Source/CameraCalibrationEditor/Private/Calibrators/CameraCalibrationSolver.h
+++ b/Engine/Plugins/VirtualProduction/CameraCalibration/Source/CameraCalibrationEditor/Private/Calibrators/CameraCalibrationSolver.h
@@ -56,7 +56,7 @@ enum class ECalibrationFlags : uint32
 	FixFocalLength = 1 << 2,     /** The solver will not optimize the focal length */
 	FixPrincipalPoint = 1 << 3,  /** The solver will not optimize the principal point */
 	FixExtrinsics = 1 << 4,      /** The solver will not optimize the camera extrinsics */
-	FixZeroDistortion = 1 << 5,  /** The solver will fix all distortion values at 0 */
+	FixDistortion = 1 << 5,      /** The solver will fix all distortion values at 0 */
 	FixAspectRatio = 1 << 6      /** The solver will respect the input aspect ratio when solving for Fx and Fy */
 };
 
@@ -77,6 +77,7 @@ public:
 		const FIntPoint ImageSize,
 		const FVector2D& FocalLength,
 		const FVector2D& ImageCenter,
+		const TArray<float>& DistortionParameters,
 		const TArray<FTransform>& CameraPoses,
 		TSubclassOf<ULensModel> LensModel,
 		double PixelAspect,
@@ -103,6 +104,7 @@ public:
 		const FIntPoint ImageSize,
 		const FVector2D& FocalLength,
 		const FVector2D& ImageCenter,
+		const TArray<float>& DistortionParameters,
 		const TArray<FTransform>& CameraPoses,
 		TSubclassOf<ULensModel> LensModel,
 		double PixelAspect,
@@ -152,6 +154,7 @@ public:
 		const FIntPoint ImageSize,
 		const FVector2D& FocalLength,
 		const FVector2D& ImageCenter,
+		const TArray<float>& DistortionParameters,
 		const TArray<FTransform>& CameraPoses,
 		TSubclassOf<ULensModel> LensModel,
 		double PixelAspect,
diff --git a/Engine/Plugins/VirtualProduction/CameraCalibration/Source/CameraCalibrationEditor/Private/Calibrators/CameraLensDistortionAlgoAruco.cpp b/Engine/Plugins/VirtualProduction/CameraCalibration/Source/CameraCalibrationEditor/Private/Calibrators/CameraLensDistortionAlgoAruco.cpp
--- a/Engine/Plugins/VirtualProduction/CameraCalibration/Source/CameraCalibrationEditor/Private/Calibrators/CameraLensDistortionAlgoAruco.cpp
+++ b/Engine/Plugins/VirtualProduction/CameraCalibration/Source/CameraCalibrationEditor/Private/Calibrators/CameraLensDistortionAlgoAruco.cpp
@@ -36,8 +36,9 @@
 
 #if WITH_EDITOR
 static TAutoConsoleVariable<bool> CVarFixExtrinsicsAruco(TEXT("LensDistortionAruco.FixExtrinsics"), false, TEXT("If true, the solver will fix the camera extrinsics to the user-provided camera poses"));
-static TAutoConsoleVariable<bool> CVarFixZeroDistortionAruco(TEXT("LensDistortionAruco.FixZeroDistortion"), false, TEXT("If true, the solver will fix all distortion values to always be 0"));
+static TAutoConsoleVariable<bool> CVarFixDistortionAruco(TEXT("LensDistortionAruco.FixDistortion"), false, TEXT("If true, the solver will not optimize distortion, and will use the input distortion values if any are given, or assume zero distortion otherwise."));
 static TAutoConsoleVariable<bool> CVarUseExtrinsicsGuessAruco(TEXT("LensDistortionAruco.UseExtrinsicsGuess"), false, TEXT("If true, the actual calibrator and camera poses will be used when running the solver"));
+static TAutoConsoleVariable<bool> CVarGroupPointsByCameraPoseAruco(TEXT("LensDistortionAruco.GroupPointsByCameraPose"), true, TEXT("If true, the points sent to the solver will be grouped together if they share the same camera pose."));
 #endif
 
 const int UCameraLensDistortionAlgoAruco::DATASET_VERSION = 1;
@@ -366,6 +367,11 @@ FDistortionCalibrationTask UCameraLensDistortionAlgoAruco::BeginCalibration(FTex
 		CameraPoses.Add(Row->CameraPose);
 	}
 
+	if (CVarGroupPointsByCameraPoseAruco.GetValueOnGameThread())
+	{
+		UE::CameraCalibration::Private::GroupPointsByCameraPose(Samples3d, Samples2d, CameraPoses);
+	}
+
 	ECalibrationFlags SolverFlags = ECalibrationFlags::None;
 
 	// Aruco markers may be detected anywhere in the image, and there is no guarantee that they will be coplanar. 
@@ -383,9 +389,9 @@ FDistortionCalibrationTask UCameraLensDistortionAlgoAruco::BeginCalibration(FTex
 		EnumAddFlags(SolverFlags, ECalibrationFlags::FixExtrinsics);
 	}
 
-	if (CVarFixZeroDistortionAruco.GetValueOnAnyThread())
+	if (CVarFixDistortionAruco.GetValueOnAnyThread())
 	{
-		EnumAddFlags(SolverFlags, ECalibrationFlags::FixZeroDistortion);
+		EnumAddFlags(SolverFlags, ECalibrationFlags::FixDistortion);
 	}
 
 	if (bFixFocalLength)
@@ -398,11 +404,14 @@ FDistortionCalibrationTask UCameraLensDistortionAlgoAruco::BeginCalibration(FTex
 		EnumAddFlags(SolverFlags, ECalibrationFlags::FixPrincipalPoint);
 	}
 
+	FDistortionInfo DistortionGuess;
+	LensFile->GetDistortionPoint(Focus, Zoom, DistortionGuess);
+
 	const TSubclassOf<ULensModel> Model = LensFile->LensInfo.LensModel;
 
 	UClass* SolverClass = LensDistortionTool->GetSolverClass();
 
-	CalibrationTask = UE::Tasks::Launch(UE_SOURCE_LOCATION, [SolverClass, Model, Samples3d, Samples2d, ImageSize, FocalLength, ImageCenter, CameraPoses, PixelAspect, SolverFlags, Focus, Zoom]() mutable
+	CalibrationTask = UE::Tasks::Launch(UE_SOURCE_LOCATION, [SolverClass, Model, Samples3d, Samples2d, ImageSize, FocalLength, ImageCenter, DistortionGuess, CameraPoses, PixelAspect, SolverFlags, Focus, Zoom]() mutable
 		{
 			ULensDistortionSolver* Solver = NewObject<ULensDistortionSolver>(GetTransientPackage(), SolverClass);
 
@@ -412,6 +421,7 @@ FDistortionCalibrationTask UCameraLensDistortionAlgoAruco::BeginCalibration(FTex
 				ImageSize,
 				FocalLength,
 				ImageCenter,
+				DistortionGuess.Parameters,
 				CameraPoses,
 				Model,
 				PixelAspect,
diff --git a/Engine/Plugins/VirtualProduction/CameraCalibration/Source/CameraCalibrationEditor/Private/Calibrators/CameraLensDistortionAlgoCheckerboard.cpp b/Engine/Plugins/VirtualProduction/CameraCalibration/Source/CameraCalibrationEditor/Private/Calibrators/CameraLensDistortionAlgoCheckerboard.cpp
--- a/Engine/Plugins/VirtualProduction/CameraCalibration/Source/CameraCalibrationEditor/Private/Calibrators/CameraLensDistortionAlgoCheckerboard.cpp
+++ b/Engine/Plugins/VirtualProduction/CameraCalibration/Source/CameraCalibrationEditor/Private/Calibrators/CameraLensDistortionAlgoCheckerboard.cpp
@@ -39,9 +39,10 @@
 #if WITH_EDITOR
 static TAutoConsoleVariable<bool> CVarUseIntrinsicsGuess(TEXT("LensDistortionCheckerboard.UseIntrinsicsGuess"), true, TEXT("If true, the solver initializes the camera intrinsics to a user-provided estimate. Otherwise, the solver will compute the initial values."));
 static TAutoConsoleVariable<bool> CVarFixExtrinsics(TEXT("LensDistortionCheckerboard.FixExtrinsics"), false, TEXT("If true, the solver will fix the camera extrinsics to the user-provided camera poses"));
-static TAutoConsoleVariable<bool> CVarFixZeroDistortion(TEXT("LensDistortionCheckerboard.FixZeroDistortion"), false, TEXT("If true, the solver will fix all distortion values to always be 0"));
+static TAutoConsoleVariable<bool> CVarFixDistortion(TEXT("LensDistortionCheckerboard.FixDistortion"), false, TEXT("If true, the solver will not optimize distortion, and will use the input distortion values if any are given, or assume zero distortion otherwise."));
 static TAutoConsoleVariable<bool> CVarUseExtrinsicsGuess(TEXT("LensDistortionCheckerboard.UseExtrinsicsGuess"), false, TEXT("If true, the actual checkerboard and camera poses will be used when running the solver"));
 static TAutoConsoleVariable<bool> CVarUseTrackedCalibrator(TEXT("LensDistortionCheckerboard.UseTrackedCalibrator"), false, TEXT("If true, the 3D points of the checkerboard will come from the tracked pose of the object, otherwise, dummy points will be used."));
+static TAutoConsoleVariable<bool> CVarGroupPointsByCameraPose(TEXT("LensDistortionCheckerboard.GroupPointsByCameraPose"), true, TEXT("If true, the points sent to the solver will be grouped together if they share the same camera pose."));
 #endif
 
 const int UCameraLensDistortionAlgoCheckerboard::DATASET_VERSION = 1;
@@ -715,6 +716,10 @@ FDistortionCalibrationTask UCameraLensDistortionAlgoCheckerboard::BeginCalibrati
 			Samples3d.Add(Points3d);
 		}
 	}
+	else if (CVarGroupPointsByCameraPose.GetValueOnGameThread())
+	{
+		UE::CameraCalibration::Private::GroupPointsByCameraPose(Samples3d, Samples2d, CameraPoses);
+	}
 
 	ECalibrationFlags SolverFlags = ECalibrationFlags::None;
 
@@ -733,9 +738,9 @@ FDistortionCalibrationTask UCameraLensDistortionAlgoCheckerboard::BeginCalibrati
 		EnumAddFlags(SolverFlags, ECalibrationFlags::FixExtrinsics);
 	}
 
-	if (CVarFixZeroDistortion.GetValueOnGameThread())
+	if (CVarFixDistortion.GetValueOnGameThread())
 	{
-		EnumAddFlags(SolverFlags, ECalibrationFlags::FixZeroDistortion);
+		EnumAddFlags(SolverFlags, ECalibrationFlags::FixDistortion);
 	}
 
 	if (bFixFocalLength)
@@ -748,11 +753,14 @@ FDistortionCalibrationTask UCameraLensDistortionAlgoCheckerboard::BeginCalibrati
 		EnumAddFlags(SolverFlags, ECalibrationFlags::FixPrincipalPoint);
 	}
 
+	FDistortionInfo DistortionGuess;
+	LensFile->GetDistortionPoint(Focus, Zoom, DistortionGuess);
+
 	const TSubclassOf<ULensModel> Model = LensFile->LensInfo.LensModel;
 	
 	Solver = NewObject<ULensDistortionSolver>(this, Tool->GetSolverClass());
 
-	CalibrationTask = UE::Tasks::Launch(UE_SOURCE_LOCATION, [Solver = Solver, Model, Samples3d, Samples2d, ImageSize, FocalLength, ImageCenter, CameraPoses, PixelAspect, SolverFlags, Focus, Zoom]() mutable
+	CalibrationTask = UE::Tasks::Launch(UE_SOURCE_LOCATION, [Solver = Solver, Model, Samples3d, Samples2d, ImageSize, FocalLength, ImageCenter, DistortionGuess, CameraPoses, PixelAspect, SolverFlags, Focus, Zoom]() mutable
 		{
 			FDistortionCalibrationResult Result = Solver->Solve(
 				Samples3d,
@@ -760,6 +768,7 @@ FDistortionCalibrationTask UCameraLensDistortionAlgoCheckerboard::BeginCalibrati
 				ImageSize,
 				FocalLength,
 				ImageCenter,
+				DistortionGuess.Parameters,
 				CameraPoses,
 				Model,
 				PixelAspect,
diff --git a/Engine/Plugins/VirtualProduction/CameraCalibration/Source/CameraCalibrationEditor/Private/Calibrators/CameraLensDistortionAlgoPoints.cpp b/Engine/Plugins/VirtualProduction/CameraCalibration/Source/CameraCalibrationEditor/Private/Calibrators/CameraLensDistortionAlgoPoints.cpp
--- a/Engine/Plugins/VirtualProduction/CameraCalibration/Source/CameraCalibrationEditor/Private/Calibrators/CameraLensDistortionAlgoPoints.cpp
+++ b/Engine/Plugins/VirtualProduction/CameraCalibration/Source/CameraCalibrationEditor/Private/Calibrators/CameraLensDistortionAlgoPoints.cpp
@@ -8,6 +8,7 @@
 #include "Camera/CameraActor.h"
 #include "CameraCalibrationEditorLog.h"
 #include "CameraCalibrationSolver.h"
+#include "CameraCalibrationUtilsPrivate.h"
 #include "Dom/JsonObject.h"
 #include "Editor.h"
 #include "Input/Events.h"
@@ -26,6 +27,13 @@
 
 #define LOCTEXT_NAMESPACE "CameraLensDistortionAlgoPoints"
 
+#if WITH_EDITOR
+static TAutoConsoleVariable<bool> CVarFixExtrinsicsPoints(TEXT("LensDistortionPoints.FixExtrinsics"), false, TEXT("If true, the solver will fix the camera extrinsics to the user-provided camera poses"));
+static TAutoConsoleVariable<bool> CVarFixDistortionPoints(TEXT("LensDistortionPoints.FixDistortion"), false, TEXT("If true, the solver will not optimize distortion, and will use the input distortion values if any are given, or assume zero distortion otherwise."));
+static TAutoConsoleVariable<bool> CVarUseExtrinsicsGuessPoints(TEXT("LensDistortionPoints.UseExtrinsicsGuess"), false, TEXT("If true, the actual checkerboard and camera poses will be used when running the solver"));
+static TAutoConsoleVariable<bool> CVarGroupPointsByCameraPosePoints(TEXT("LensDistortionPoints.GroupPointsByCameraPose"), true, TEXT("If true, the points sent to the solver will be grouped together if they share the same camera pose."));
+#endif
+
 const int UCameraLensDistortionAlgoPoints::DATASET_VERSION = 1;
 
 // String constants for import/export
@@ -184,6 +192,15 @@ void UCameraLensDistortionAlgoPoints::Tick(float DeltaTime)
 			LastCameraData.InputFocus = EvalInputs.Focus;
 			LastCameraData.InputZoom = EvalInputs.Zoom;
 
+			const UCameraComponent* CameraComponent = Camera->GetCameraComponent();
+
+			if (!CameraComponent)
+			{
+				break;
+			}
+
+			LastCameraData.Pose = CameraComponent->GetComponentToWorld();
+
 			if (!Calibrator.IsValid())
 			{
 				break;
@@ -257,7 +274,7 @@ bool UCameraLensDistortionAlgoPoints::OnViewportClicked(const FGeometry& MyGeome
 	}
 
 	// Store the 2D point in pixel coordinates
-	const FIntPoint ImageSize = StepsController->GetCompRenderResolution();
+	const FIntPoint ImageSize = StepsController->GetCameraFeedSize();
 	LastCalibratorPoint.Point2d = LastCalibratorPoint.Point2d * ImageSize;
 
 	// Export the latest session data
@@ -350,20 +367,21 @@ TSharedRef<SWidget> UCameraLensDistortionAlgoPoints::BuildUI()
 		[BuildCalibrationActionButtons()];
 }
 
-bool UCameraLensDistortionAlgoPoints::GetLensDistortion(
-	float& OutFocus,
-	float& OutZoom,
-	FDistortionInfo& OutDistortionInfo,
-	FFocalLengthInfo& OutFocalLengthInfo,
-	FImageCenterInfo& OutImageCenterInfo,
-	TSubclassOf<ULensModel>& OutLensModel,
-	double& OutError,
-	FText& OutErrorMessage)
+FDistortionCalibrationTask UCameraLensDistortionAlgoPoints::BeginCalibration(FText& OutErrorMessage)
 {
+	FDistortionCalibrationTask CalibrationTask = {};
+
+	// Validate that enough points were gathered to attempt a calibration
+	if (CalibrationRows.Num() < 1)
+	{
+		OutErrorMessage = LOCTEXT("NotEnoughCalibrationRowsError", "Could not initiate distortion calibration. At least 1 calibration row is required.");
+		return CalibrationTask;
+	}
+
 	if (!LensDistortionTool.IsValid())
 	{
 		OutErrorMessage = LOCTEXT("InvalidTool", "Invalid Tool");
-		return false;
+		return CalibrationTask;
 	}
 
 	const FCameraCalibrationStepsController* StepsController = LensDistortionTool->GetCameraCalibrationStepsController();
@@ -371,7 +389,7 @@ bool UCameraLensDistortionAlgoPoints::GetLensDistortion(
 	if (!StepsController)
 	{
 		OutErrorMessage = LOCTEXT("InvalidStepsController", "Invalid StepsController");
-		return false;
+		return CalibrationTask;
 	}
 
 	const ULensFile* LensFile = StepsController->GetLensFile();
@@ -379,9 +397,12 @@ bool UCameraLensDistortionAlgoPoints::GetLensDistortion(
 	if (!LensFile)
 	{
 		OutErrorMessage = LOCTEXT("InvalidLensFile", "Invalid Lens File");
-		return false;
+		return CalibrationTask;
 	}
 
+	const float Focus = CalibrationRows[0]->CameraData.InputFocus;
+	const float Zoom = CalibrationRows[0]->CameraData.InputZoom;
+
 	TArray<FObjectPoints> Samples3d;
 	TArray<FImagePoints> Samples2d;
 	Samples3d.Reserve(CalibrationRows.Num());
@@ -392,6 +413,8 @@ bool UCameraLensDistortionAlgoPoints::GetLensDistortion(
 	uint32 NextPatternIndex = 0;
 	for (int32 RowIndex = 0; RowIndex < CalibrationRows.Num(); ++RowIndex)
 	{
+		FTransform Pose = CalibrationRows[RowIndex]->CameraData.Pose;
+
 		FObjectPoints Points3d;
 		FImagePoints Points2d;
 		while (RowIndex < CalibrationRows.Num() && CalibrationRows[RowIndex]->PatternIndex == NextPatternIndex)
@@ -410,7 +433,7 @@ bool UCameraLensDistortionAlgoPoints::GetLensDistortion(
 		{
 			Samples3d.Add(Points3d);
 			Samples2d.Add(Points2d);
-			CameraPoses.Add(FTransform::Identity);
+			CameraPoses.Add(Pose);
 		}
 
 		--RowIndex;
@@ -421,7 +444,7 @@ bool UCameraLensDistortionAlgoPoints::GetLensDistortion(
 	if (Samples3d.Num() < 4)
 	{
 		OutErrorMessage = LOCTEXT("NotEnoughSamples", "At least 4 calibration patterns are required");
-		return false;
+		return CalibrationTask;
 	}
 
 	// Validate that each pattern has the same number of points
@@ -431,16 +454,21 @@ bool UCameraLensDistortionAlgoPoints::GetLensDistortion(
 		if (Samples3d[PatternIndex].Points.Num() != NumPointsInPattern)
 		{
 			OutErrorMessage = LOCTEXT("DifferentNumPointsInPattern", "Every calibration pattern must have the same number of points");
-			return false;
+			return CalibrationTask;
 		}
 	}
 
+	if (CVarGroupPointsByCameraPosePoints.GetValueOnGameThread())
+	{
+		UE::CameraCalibration::Private::GroupPointsByCameraPose(Samples3d, Samples2d, CameraPoses);
+	}
+
 	// Because the calibration pattern is not coplanar, OpenCV requires an initial intrinsics guess
 	if (FMath::IsNearlyEqual(FocalLengthEstimate, 0.0) || FocalLengthEstimate < 0.0)
 	{
 		OutErrorMessage = LOCTEXT("InvalidFocalLengthMsg", "The current focal length (in mm) must be set to a "
 			"valid value in order to seed the calibration algorithm with a best guess as to the camera intrinsics (Fx/Fy)");
-		return false;
+		return CalibrationTask;
 	}
 
 	// Validate sensor dimensions
@@ -451,7 +479,7 @@ bool UCameraLensDistortionAlgoPoints::GetLensDistortion(
 	if (FMath::IsNearlyZero(PixelAspect))
 	{
 		OutErrorMessage = LOCTEXT("PixelAspectNearlyZero", "The pixel aspect in the Lens Information is zero, which is invalid.");
-		return false;
+		return CalibrationTask;
 	}
 
 	const float PhysicalSensorWidth = StepsController->GetLensFileEvaluationInputs().Filmback.SensorWidth;
@@ -460,7 +488,7 @@ bool UCameraLensDistortionAlgoPoints::GetLensDistortion(
 	if (FMath::IsNearlyZero(DesqueezeSensorWidth))
 	{
 		OutErrorMessage = LOCTEXT("SensorWidthNearlyZero", "The sensor width of the CineCamera is zero, which is invalid.");
-		return false;
+		return CalibrationTask;
 	}
 
 	const double Fx = (FocalLengthEstimate / DesqueezeSensorWidth) * ImageSize.X;
@@ -472,41 +500,72 @@ bool UCameraLensDistortionAlgoPoints::GetLensDistortion(
 	ECalibrationFlags SolverFlags = ECalibrationFlags::None;
 	EnumAddFlags(SolverFlags, ECalibrationFlags::UseIntrinsicGuess);
 
-	UClass* SolverClass = LensDistortionTool->GetSolverClass();
-	ULensDistortionSolver* Solver = NewObject<ULensDistortionSolver>(this, SolverClass);
-
-	FDistortionCalibrationResult Result = Solver->Solve(
-		Samples3d,
-		Samples2d,
-		ImageSize,
-		FocalLength,
-		ImageCenter,
-		CameraPoses,
-		LensFile->LensInfo.LensModel,
-		PixelAspect,
-		SolverFlags
-	);
-
-	if (!Result.ErrorMessage.IsEmpty())
+	if (CVarUseExtrinsicsGuessPoints.GetValueOnGameThread())
 	{
-		OutErrorMessage = Result.ErrorMessage;
-		return false;
+		EnumAddFlags(SolverFlags, ECalibrationFlags::UseExtrinsicGuess);
 	}
 
-	const TSharedPtr<FLensDistortionPointsRowData>& FirstRow = CalibrationRows[0];
+	if (CVarFixExtrinsicsPoints.GetValueOnGameThread())
+	{
+		EnumAddFlags(SolverFlags, ECalibrationFlags::FixExtrinsics);
+	}
 
-	OutLensModel = LensFile->LensInfo.LensModel;
-	OutDistortionInfo.Parameters = Result.Parameters.Parameters;
+	if (CVarFixDistortionPoints.GetValueOnGameThread())
+	{
+		EnumAddFlags(SolverFlags, ECalibrationFlags::FixDistortion);
+	}
 
-	OutFocalLengthInfo.FxFy = FVector2D(Result.FocalLength.FxFy / ImageSize);
+	FDistortionInfo DistortionGuess;
+	LensFile->GetDistortionPoint(Focus, Zoom, DistortionGuess);
 
-	OutImageCenterInfo.PrincipalPoint = FVector2D(Result.ImageCenter.PrincipalPoint / ImageSize);
+	const TSubclassOf<ULensModel> Model = LensFile->LensInfo.LensModel;
 
-	// FZ inputs to LUT
-	OutFocus = FirstRow->CameraData.InputFocus;
-	OutZoom = FirstRow->CameraData.InputZoom;
+	Solver = NewObject<ULensDistortionSolver>(this, LensDistortionTool->GetSolverClass());
 
-	return true;
+	CalibrationTask = UE::Tasks::Launch(UE_SOURCE_LOCATION, [Solver = Solver, Model, Samples3d, Samples2d, ImageSize, FocalLength, ImageCenter, DistortionGuess, CameraPoses, PixelAspect, SolverFlags, Focus, Zoom]() mutable
+	{
+		FDistortionCalibrationResult Result = Solver->Solve(
+			Samples3d,
+			Samples2d,
+			ImageSize,
+			FocalLength,
+			ImageCenter,
+			DistortionGuess.Parameters,
+			CameraPoses,
+			Model,
+			PixelAspect,
+			SolverFlags
+		);
+
+		// CalibrateCamera() returns focal length and image center in pixels, but the result is expected to be normalized by the image size
+		Result.FocalLength.FxFy = Result.FocalLength.FxFy / ImageSize;
+		Result.ImageCenter.PrincipalPoint = Result.ImageCenter.PrincipalPoint / ImageSize;
+
+		// FZ inputs to LUT
+		Result.EvaluatedFocus = Focus;
+		Result.EvaluatedZoom = Zoom;
+
+		return Result;
+	});
+
+	return CalibrationTask;
+}
+
+void UCameraLensDistortionAlgoPoints::CancelCalibration()
+{
+	if (Solver)
+	{
+		Solver->Cancel();
+	}
+}
+
+bool UCameraLensDistortionAlgoPoints::GetCalibrationStatus(FText& StatusText) const
+{
+	if (Solver)
+	{
+		return Solver->GetStatusText(StatusText);
+	}
+	return false;
 }
 
 void UCameraLensDistortionAlgoPoints::OnDistortionSavedToLens()
diff --git a/Engine/Plugins/VirtualProduction/CameraCalibration/Source/CameraCalibrationEditor/Private/Calibrators/CameraLensDistortionAlgoPoints.h b/Engine/Plugins/VirtualProduction/CameraCalibration/Source/CameraCalibrationEditor/Private/Calibrators/CameraLensDistortionAlgoPoints.h
--- a/Engine/Plugins/VirtualProduction/CameraCalibration/Source/CameraCalibrationEditor/Private/Calibrators/CameraLensDistortionAlgoPoints.h
+++ b/Engine/Plugins/VirtualProduction/CameraCalibration/Source/CameraCalibrationEditor/Private/Calibrators/CameraLensDistortionAlgoPoints.h
@@ -4,6 +4,7 @@
 
 #include "CameraLensDistortionAlgo.h"
 
+#include "CameraCalibrationSolver.h"
 #include "LensFile.h"
 
 #include "CameraLensDistortionAlgoPoints.generated.h"
@@ -59,6 +60,10 @@ struct FLensDistortionPointsCameraData
 	UPROPERTY()
 	uint32 CalibratorUniqueId = INDEX_NONE;
 
+	// The camera pose
+	UPROPERTY()
+	FTransform Pose = FTransform::Identity;
+
 	// Input focus from lens file evaluation data
 	UPROPERTY()
 	float InputFocus = 0.0f;
@@ -119,15 +124,10 @@ public:
 	virtual FName ShortName() const override { return TEXT("Points"); };
 	virtual void OnDistortionSavedToLens() override;
 
-	virtual bool GetLensDistortion(
-		float& OutFocus,
-		float& OutZoom,
-		FDistortionInfo& OutDistortionInfo,
-		FFocalLengthInfo& OutFocalLengthInfo,
-		FImageCenterInfo& OutImageCenterInfo,
-		TSubclassOf<ULensModel>& OutLensModel,
-		double& OutError,
-		FText& OutErrorMessage) override;
+	virtual FDistortionCalibrationTask BeginCalibration(FText& OutErrorMessage) override;
+	virtual void CancelCalibration() override;
+	virtual bool GetCalibrationStatus(FText& StatusText) const override;
+	virtual bool SupportsAsyncCalibration() override { return true; };
 
 	virtual bool HasCalibrationData() const override;
 	virtual void PreImportCalibrationData() override;
@@ -190,6 +190,9 @@ protected:
 	/** The value of the current focal length of the lens (in mm), used to set a guess for the camera intrinsics fx/fy */
 	double FocalLengthEstimate = 0.0;
 
+	/** Solver instance that will run the distortion calibration */
+	TObjectPtr<ULensDistortionSolver> Solver;
+
 protected:
 
 	/** Returns the first calibrator object in the scene that it can find */
diff --git a/Engine/Plugins/VirtualProduction/CameraCalibration/Source/CameraCalibrationEditor/Private/CameraCalibrationUtilsPrivate.cpp b/Engine/Plugins/VirtualProduction/CameraCalibration/Source/CameraCalibrationEditor/Private/CameraCalibrationUtilsPrivate.cpp
--- a/Engine/Plugins/VirtualProduction/CameraCalibration/Source/CameraCalibrationEditor/Private/CameraCalibrationUtilsPrivate.cpp
+++ b/Engine/Plugins/VirtualProduction/CameraCalibration/Source/CameraCalibrationEditor/Private/CameraCalibrationUtilsPrivate.cpp
@@ -3,6 +3,8 @@
 #include "CameraCalibrationUtilsPrivate.h"
 
 #include "CalibrationPointComponent.h"
+#include "Calibrators/CameraCalibrationSolver.h"
+#include "CameraCalibrationUtils.h"
 #include "Editor.h"
 #include "Engine/Level.h"
 #include "Engine/Texture2D.h"
@@ -200,4 +202,53 @@ namespace UE::CameraCalibration::Private
 			Texture->UpdateResource();
 		}
 	}
+
+	void GroupPointsByCameraPose(TArray<FObjectPoints>& InOutObjectPoints, TArray<FImagePoints>& InOutImagePoints, TArray<FTransform>& InOutCameraPoses)
+	{
+		const int NumPoses = InOutCameraPoses.Num();
+		if (!ensure(NumPoses == InOutObjectPoints.Num()) && (NumPoses == InOutImagePoints.Num()))
+		{
+			return;
+		}
+
+		TArray<FObjectPoints> GroupedObjectPoints;
+		TArray<FImagePoints> GroupedImagePoints;
+		TArray<FTransform> GroupedCameraPoses;
+
+		FTransform LastCameraPose = InOutCameraPoses[0];
+		FTransform CurrentCameraPose;
+
+		FObjectPoints CurrentObjectPoints;
+		FImagePoints CurrentImagePoints;
+		for (int32 PoseIndex = 0; PoseIndex < NumPoses; ++PoseIndex)
+		{
+			CurrentCameraPose = InOutCameraPoses[PoseIndex];
+
+			if (!FCameraCalibrationUtils::IsNearlyEqual(CurrentCameraPose, LastCameraPose))
+			{
+				GroupedCameraPoses.Add(LastCameraPose);
+				LastCameraPose = CurrentCameraPose;
+
+				GroupedObjectPoints.Add(CurrentObjectPoints);
+				CurrentObjectPoints.Points.Reset();
+
+				GroupedImagePoints.Add(CurrentImagePoints);
+				CurrentImagePoints.Points.Reset();
+			}
+
+			CurrentObjectPoints.Points.Append(InOutObjectPoints[PoseIndex].Points);
+			CurrentImagePoints.Points.Append(InOutImagePoints[PoseIndex].Points);
+		}
+
+		if (!CurrentObjectPoints.Points.IsEmpty())
+		{
+			GroupedObjectPoints.Add(CurrentObjectPoints);
+			GroupedImagePoints.Add(CurrentImagePoints);
+			GroupedCameraPoses.Add(LastCameraPose);
+		}
+
+		InOutObjectPoints = GroupedObjectPoints;
+		InOutImagePoints = GroupedImagePoints;
+		InOutCameraPoses = GroupedCameraPoses;
+	}
 }
diff --git a/Engine/Plugins/VirtualProduction/CameraCalibration/Source/CameraCalibrationEditor/Private/CameraCalibrationUtilsPrivate.h b/Engine/Plugins/VirtualProduction/CameraCalibration/Source/CameraCalibrationEditor/Private/CameraCalibrationUtilsPrivate.h
--- a/Engine/Plugins/VirtualProduction/CameraCalibration/Source/CameraCalibrationEditor/Private/CameraCalibrationUtilsPrivate.h
+++ b/Engine/Plugins/VirtualProduction/CameraCalibration/Source/CameraCalibrationEditor/Private/CameraCalibrationUtilsPrivate.h
@@ -7,6 +7,9 @@
 
 class UCalibrationPointComponent;
 
+struct FObjectPoints;
+struct FImagePoints;
+
 /** Structure representing an aruco marker, including ID, Name, and corner coordinates in 2D image space and 3D world space */
 USTRUCT()
 struct FArucoCalibrationPoint
@@ -52,4 +55,7 @@ namespace UE::CameraCalibration::Private
 
 	/** Set the texture data to the input array of pixels */
 	void SetTextureData(UTexture2D* Texture, const TArray<FColor>& PixelData);
+
+	/** Group the input points together if they share the same camera pose, and output the unique camera poses */
+	void GroupPointsByCameraPose(TArray<FObjectPoints>& InOutObjectPoints, TArray<FImagePoints>& InOutImagePoints, TArray<FTransform>& InOutCameraPoses);
 }
diff --git a/Engine/Plugins/VirtualProduction/CameraCalibration/Source/CameraCalibrationEditor/Private/TestCameraCalibration.cpp b/Engine/Plugins/VirtualProduction/CameraCalibration/Source/CameraCalibrationEditor/Private/TestCameraCalibration.cpp
--- a/Engine/Plugins/VirtualProduction/CameraCalibration/Source/CameraCalibrationEditor/Private/TestCameraCalibration.cpp
+++ b/Engine/Plugins/VirtualProduction/CameraCalibration/Source/CameraCalibrationEditor/Private/TestCameraCalibration.cpp
@@ -127,7 +127,7 @@ namespace UE::Private::CameraCalibration::AutomatedTests
 			UE_LOG(LogTestCameraCalibration, Log, TEXT("Fix Focal Length: %s"), *GetSolverFlagString(Test.SolverSettings.bFixFocalLength));
 			UE_LOG(LogTestCameraCalibration, Log, TEXT("Fix Image Center: %s"), *GetSolverFlagString(Test.SolverSettings.bFixPrincipalPoint));
 			UE_LOG(LogTestCameraCalibration, Log, TEXT("Fix Extrinsics: %s"), *GetSolverFlagString(Test.SolverSettings.bFixExtrinsics));
-			UE_LOG(LogTestCameraCalibration, Log, TEXT("Fix Zero Distortion: %s"), *GetSolverFlagString(Test.SolverSettings.bFixZeroDistortion));
+			UE_LOG(LogTestCameraCalibration, Log, TEXT("Fix Distortion: %s"), *GetSolverFlagString(Test.SolverSettings.bFixDistortion));
 			UE_LOG(LogTestCameraCalibration, Log, TEXT("Fix Aspect Ratio: %s"), *GetSolverFlagString(Test.SolverSettings.bFixAspectRatio));
 		}
 
@@ -175,9 +175,9 @@ namespace UE::Private::CameraCalibration::AutomatedTests
 		{
 			EnumAddFlags(SolverFlags, ECalibrationFlags::FixExtrinsics);
 		}
-		if (TestSettings->bFixZeroDistortion)
+		if (TestSettings->bFixDistortion)
 		{
-			EnumAddFlags(SolverFlags, ECalibrationFlags::FixZeroDistortion);
+			EnumAddFlags(SolverFlags, ECalibrationFlags::FixDistortion);
 		}
 
 		/** 
@@ -348,7 +348,7 @@ namespace UE::Private::CameraCalibration::AutomatedTests
 		}
 
 		FVector2D CalibratedImageCenter = ImageCenter;
-		TArray<float> CalibratedDistortionParameters;
+		TArray<float> DistortionGuess;
 
 		ULensDistortionSolverOpenCV* TestSolver = NewObject<ULensDistortionSolverOpenCV>();
 
@@ -358,6 +358,7 @@ namespace UE::Private::CameraCalibration::AutomatedTests
 			ImageSize,
 			CalibratedFxFy,
 			CalibratedImageCenter,
+			DistortionGuess,
 			EstimatedCameraPoses,
 			USphericalLensModel::StaticClass(),
 			PixelAspect,
@@ -788,9 +789,9 @@ namespace UE::Private::CameraCalibration::AutomatedTests
 		{
 			EnumAddFlags(SolverFlags, ECalibrationFlags::FixExtrinsics);
 		}
-		if (TestOptional(SolverSettings.bFixZeroDistortion))
+		if (TestOptional(SolverSettings.bFixDistortion))
 		{
-			EnumAddFlags(SolverFlags, ECalibrationFlags::FixZeroDistortion);
+			EnumAddFlags(SolverFlags, ECalibrationFlags::FixDistortion);
 		}
 		if (TestOptional(SolverSettings.bFixAspectRatio))
 		{
@@ -841,6 +842,9 @@ namespace UE::Private::CameraCalibration::AutomatedTests
 		TArray<FTransform> CameraPoseGuesses;
 		CameraPoseGuesses.Empty();
 
+		TArray<float> InitialDistortion;
+		InitialDistortion.Empty();
+
 		const FVector2D FocalLengthGuessInPixels = CalibrationTest.CameraProfile.ConvertFocalLengthToPixels(SolverSettings.FocalLengthGuess);
 
 		FDistortionCalibrationResult Result = TestSolver->Solve(
@@ -849,6 +853,7 @@ namespace UE::Private::CameraCalibration::AutomatedTests
 			CalibrationTest.CameraProfile.ImageSize,
 			FocalLengthGuessInPixels,
 			SolverSettings.ImageCenterGuess,
+			InitialDistortion,
 			CameraPoseGuesses,
 			USphericalLensModel::StaticClass(),
 			1.0,
diff --git a/Engine/Plugins/VirtualProduction/CameraCalibration/Source/CameraCalibrationEditor/Private/TestCameraCalibration.h b/Engine/Plugins/VirtualProduction/CameraCalibration/Source/CameraCalibrationEditor/Private/TestCameraCalibration.h
--- a/Engine/Plugins/VirtualProduction/CameraCalibration/Source/CameraCalibrationEditor/Private/TestCameraCalibration.h
+++ b/Engine/Plugins/VirtualProduction/CameraCalibration/Source/CameraCalibrationEditor/Private/TestCameraCalibration.h
@@ -45,10 +45,10 @@ struct FLocationRotation
 	GENERATED_BODY()
 
 	UPROPERTY()
-	FVector Location;
+	FVector Location = FVector::ZeroVector;
 
 	UPROPERTY()
-	FRotator Rotation;
+	FRotator Rotation = FRotator::ZeroRotator;
 };
 
 /** Description of a camera / lens, useful for projecting 3D calibrator points to 2D */
@@ -129,7 +129,7 @@ struct FSolverSettings
 	TOptional<bool> bFixExtrinsics;
 	
 	UPROPERTY()
-	TOptional<bool> bFixZeroDistortion;
+	TOptional<bool> bFixDistortion;
 	
 	UPROPERTY()
 	TOptional<bool> bFixAspectRatio;
diff --git a/Engine/Plugins/VirtualProduction/CameraCalibration/Source/CameraCalibrationEditor/Private/TestCameraCalibrationSettings.h b/Engine/Plugins/VirtualProduction/CameraCalibration/Source/CameraCalibrationEditor/Private/TestCameraCalibrationSettings.h
--- a/Engine/Plugins/VirtualProduction/CameraCalibration/Source/CameraCalibrationEditor/Private/TestCameraCalibrationSettings.h
+++ b/Engine/Plugins/VirtualProduction/CameraCalibration/Source/CameraCalibrationEditor/Private/TestCameraCalibrationSettings.h
@@ -101,6 +101,10 @@ public:
 	UPROPERTY(config, EditAnywhere, Category = "Calibrate Camera Solver Settings")
 	bool bFixExtrinsics = false;
 
+	/** If true, the solver will not try to solve for distortion */
+	UPROPERTY(config, EditAnywhere, Category = "Calibrate Camera Solver Settings")
+	bool bFixDistortion = false;
+
 	/** If true, the solver will not try to solve for distortion, and will assume there is zero distortion in the image */
 	UPROPERTY(config, EditAnywhere, Category = "Calibrate Camera Solver Settings")
 	bool bFixZeroDistortion = false;
diff --git a/Engine/Plugins/VirtualProduction/CameraCalibrationCore/Source/CameraCalibrationCore/Public/CameraCalibrationTypes.h b/Engine/Plugins/VirtualProduction/CameraCalibrationCore/Source/CameraCalibrationCore/Public/CameraCalibrationTypes.h
--- a/Engine/Plugins/VirtualProduction/CameraCalibrationCore/Source/CameraCalibrationCore/Public/CameraCalibrationTypes.h
+++ b/Engine/Plugins/VirtualProduction/CameraCalibrationCore/Source/CameraCalibrationCore/Public/CameraCalibrationTypes.h
@@ -38,6 +38,10 @@ struct FDistortionCalibrationResult
 	UPROPERTY(BlueprintReadWrite, Category = "Distortion")
 	TArray<FTransform> CameraPoses;
 
+	/** Calibrated nodal offset result */
+	UPROPERTY(BlueprintReadWrite, Category = "Distortion")
+	FNodalPointOffset NodalOffset;
+
 	/** Distortion parameters for the model specified by the lens file. And empty parameter array implies that there is a valid ST Map instead. */
 	UPROPERTY(BlueprintReadWrite, Category = "Distortion")
 	FDistortionInfo Parameters;
