Fixed unused or potentially uninitialized variables.

diff --git a/Engine/Plugins/Experimental/Avalanche/Source/AvalancheLevelViewport/Private/Customizations/AvaLevelViewportComponentTransformDetails.cpp b/Engine/Plugins/Experimental/Avalanche/Source/AvalancheLevelViewport/Private/Customizations/AvaLevelViewportComponentTransformDetails.cpp
--- a/Engine/Plugins/Experimental/Avalanche/Source/AvalancheLevelViewport/Private/Customizations/AvaLevelViewportComponentTransformDetails.cpp
+++ b/Engine/Plugins/Experimental/Avalanche/Source/AvalancheLevelViewport/Private/Customizations/AvaLevelViewportComponentTransformDetails.cpp
@@ -1057,7 +1057,7 @@ void FAvaLevelViewportComponentTransformDetails::CacheTransform()
 	}
 
 	FVector CurLoc = FVector::ZeroVector;
-	FRotator CurRot;
+	FRotator CurRot = FRotator::ZeroRotator;
 	FVector CurScale = FVector::ZeroVector;
 	bool bFirstObject = true;
 	bool bSkipFirst = false;
diff --git a/Engine/Plugins/Experimental/Avalanche/Source/AvalancheSequencer/Private/EaseCurveTool/AvaEaseCurveTangents.cpp b/Engine/Plugins/Experimental/Avalanche/Source/AvalancheSequencer/Private/EaseCurveTool/AvaEaseCurveTangents.cpp
--- a/Engine/Plugins/Experimental/Avalanche/Source/AvalancheSequencer/Private/EaseCurveTool/AvaEaseCurveTangents.cpp
+++ b/Engine/Plugins/Experimental/Avalanche/Source/AvalancheSequencer/Private/EaseCurveTool/AvaEaseCurveTangents.cpp
@@ -288,7 +288,7 @@ double FAvaEaseCurveTangents::CalculateCurveLength(const int32 SampleCount) cons
 
 	// Sample points along the curve and use those points to calculate the length of the curve
 	FVector2d SamplePoint;
-	FVector2d PreviousSamplePoint;
+	FVector2d PreviousSamplePoint = FVector2d::ZeroVector;
 	double Length = 0.0;
 
 	for (int32 Index = 0; Index <= SampleCount; Index++)
diff --git a/Engine/Plugins/Importers/USDImporter/Source/USDUtilities/Private/USDShadeConversion.cpp b/Engine/Plugins/Importers/USDImporter/Source/USDUtilities/Private/USDShadeConversion.cpp
--- a/Engine/Plugins/Importers/USDImporter/Source/USDUtilities/Private/USDShadeConversion.cpp
+++ b/Engine/Plugins/Importers/USDImporter/Source/USDUtilities/Private/USDShadeConversion.cpp
@@ -1253,7 +1253,7 @@ namespace UE
 						);
 					}
 
-					FLinearColor ComponentMask;
+					FLinearColor ComponentMask = FLinearColor::Black;
 					switch (TextureValue.OutputIndex)
 					{
 						case 0:	   // RGB
diff --git a/Engine/Plugins/Runtime/nDisplay/Source/DisplayClusterMoviePipeline/Private/DisplayClusterMoviePipelineViewportPass.cpp b/Engine/Plugins/Runtime/nDisplay/Source/DisplayClusterMoviePipeline/Private/DisplayClusterMoviePipelineViewportPass.cpp
--- a/Engine/Plugins/Runtime/nDisplay/Source/DisplayClusterMoviePipeline/Private/DisplayClusterMoviePipelineViewportPass.cpp
+++ b/Engine/Plugins/Runtime/nDisplay/Source/DisplayClusterMoviePipeline/Private/DisplayClusterMoviePipelineViewportPass.cpp
@@ -311,7 +311,7 @@ FSceneView* UDisplayClusterMoviePipelineViewportPassBase::GetSceneViewForSampleS
 		const FIntPoint FullDestSize(InOutSampleState.TileSize.X, InOutSampleState.TileSize.Y);
 		const FIntPoint DestSize = FullDestSize - OffsetMin - OffsetMax;
 
-		FDisplayClusterViewInfo NewView, PrevView;
+		FDisplayClusterViewInfo NewView;
 		DCViewport = GetAndCalculateDisplayClusterViewport(InOutSampleState, *ViewportId, DestSize, 0, NewView);
 		if(DCViewport)
 		{
diff --git a/Engine/Plugins/VirtualProduction/CameraCalibration/Source/CameraCalibrationEditor/Private/LensFileExchangeFormat.cpp b/Engine/Plugins/VirtualProduction/CameraCalibration/Source/CameraCalibrationEditor/Private/LensFileExchangeFormat.cpp
--- a/Engine/Plugins/VirtualProduction/CameraCalibration/Source/CameraCalibrationEditor/Private/LensFileExchangeFormat.cpp
+++ b/Engine/Plugins/VirtualProduction/CameraCalibration/Source/CameraCalibrationEditor/Private/LensFileExchangeFormat.cpp
@@ -672,7 +672,7 @@ void FLensFileExchange::PopulateFocalLengthTable(ULensFile& OutLensFile, const F
 		{
 			float Focus = 0.0f;
 			float Zoom = 0.0f;
-			FVector2D FxFy;
+			FVector2D FxFy = FVector2D::ZeroVector;
 
 			for (int32 HeaderIndex = 0; HeaderIndex < NumHeaders; ++HeaderIndex)
 			{
@@ -722,7 +722,7 @@ void FLensFileExchange::PopulateImageCenterTable(ULensFile& OutLensFile, const F
 		{
 			float Focus = 0.0f;
 			float Zoom = 0.0f;
-			FVector2D FxFy;
+			FVector2D FxFy = FVector2D::ZeroVector;
 
 			for (int32 HeaderIndex = 0; HeaderIndex < NumHeaders; ++HeaderIndex)
 			{
