CameraCalibration: Change the sensitivity of the transparency slider, and fix issue where the widget loses focus after you type in a value.

diff --git a/Engine/Plugins/VirtualProduction/CameraCalibration/Source/CameraCalibrationEditor/Private/AssetEditor/SCameraCalibrationSteps.cpp b/Engine/Plugins/VirtualProduction/CameraCalibration/Source/CameraCalibrationEditor/Private/AssetEditor/SCameraCalibrationSteps.cpp
--- a/Engine/Plugins/VirtualProduction/CameraCalibration/Source/CameraCalibrationEditor/Private/AssetEditor/SCameraCalibrationSteps.cpp
+++ b/Engine/Plugins/VirtualProduction/CameraCalibration/Source/CameraCalibrationEditor/Private/AssetEditor/SCameraCalibrationSteps.cpp
@@ -425,10 +425,7 @@ TSharedRef<SWidget> SCameraCalibrationSteps::BuildSimulcamWiperWidget()
 		})
 		.MinValue(0.0f)
 		.MaxValue(1.0f)
-		.MinSliderValue(0.0f)
-		.MaxSliderValue(1.0f)
-		.ClearKeyboardFocusOnCommit(true)
-		.Delta(0.1f);
+		.Delta(0.01f);
 }
 
 void SCameraCalibrationSteps::UpdateMediaSourcesOptions()
