Slate: Deprecate SListView::ItemHeight and STreeViewItemHeight. ItemHeight and ItemWidth are only used when the panel is is as a tile or both value are assigned. With TreeView and ListView, you can't assigned both. The function are protected, a subclass of those could but it's up to them to expose the ItemHeight and ItemWidth. This behavior changed was introduced with CL 7634094.

diff --git a/Engine/Plugins/Experimental/PropertyAnimatorCore/Source/PropertyAnimatorCoreEditor/Private/Widgets/Views/SPropertyAnimatorCoreEditorPropertiesView.cpp b/Engine/Plugins/Experimental/PropertyAnimatorCore/Source/PropertyAnimatorCoreEditor/Private/Widgets/Views/SPropertyAnimatorCoreEditorPropertiesView.cpp
--- a/Engine/Plugins/Experimental/PropertyAnimatorCore/Source/PropertyAnimatorCoreEditor/Private/Widgets/Views/SPropertyAnimatorCoreEditorPropertiesView.cpp
+++ b/Engine/Plugins/Experimental/PropertyAnimatorCore/Source/PropertyAnimatorCoreEditor/Private/Widgets/Views/SPropertyAnimatorCoreEditorPropertiesView.cpp
@@ -25,7 +25,6 @@ void SPropertyAnimatorCoreEditorPropertiesView::Construct(const FArguments& InAr
 			SAssignNew(PropertiesTree, STreeView<FPropertiesViewItemPtr>)
 		   .TreeItemsSource(&PropertiesTreeSource)
 		   .SelectionMode(ESelectionMode::None)
-		   .ItemHeight(30.f)
 		   .HeaderRow(
 			   SNew(SHeaderRow)
 			   + SHeaderRow::Column(SPropertyAnimatorCoreEditorEditPanel::HeaderPropertyColumnName)
diff --git a/Engine/Plugins/Runtime/nDisplay/Source/DisplayClusterConfigurator/Private/Views/TreeViews/SDisplayClusterConfiguratorViewTree.cpp b/Engine/Plugins/Runtime/nDisplay/Source/DisplayClusterConfigurator/Private/Views/TreeViews/SDisplayClusterConfiguratorViewTree.cpp
--- a/Engine/Plugins/Runtime/nDisplay/Source/DisplayClusterConfigurator/Private/Views/TreeViews/SDisplayClusterConfiguratorViewTree.cpp
+++ b/Engine/Plugins/Runtime/nDisplay/Source/DisplayClusterConfigurator/Private/Views/TreeViews/SDisplayClusterConfiguratorViewTree.cpp
@@ -258,7 +258,6 @@ void SDisplayClusterConfiguratorViewTree::CreateTreeColumns()
 			.OnItemScrolledIntoView(this, &SDisplayClusterConfiguratorViewTree::OnItemScrolledIntoView)
 			.OnMouseButtonDoubleClick(this, &SDisplayClusterConfiguratorViewTree::OnTreeDoubleClick)
 			.OnSetExpansionRecursive(this, &SDisplayClusterConfiguratorViewTree::SetTreeItemExpansionRecursive)
-			.ItemHeight(24)
 			.HighlightParentNodesForSelection(true)
 			.HeaderRow
 			(
diff --git a/Engine/Plugins/Runtime/nDisplay/Source/DisplayClusterLightCardEditor/Private/LightCardTemplates/SDisplayClusterLightCardTemplateList.cpp b/Engine/Plugins/Runtime/nDisplay/Source/DisplayClusterLightCardEditor/Private/LightCardTemplates/SDisplayClusterLightCardTemplateList.cpp
--- a/Engine/Plugins/Runtime/nDisplay/Source/DisplayClusterLightCardEditor/Private/LightCardTemplates/SDisplayClusterLightCardTemplateList.cpp
+++ b/Engine/Plugins/Runtime/nDisplay/Source/DisplayClusterLightCardEditor/Private/LightCardTemplates/SDisplayClusterLightCardTemplateList.cpp
@@ -239,7 +239,6 @@ void SDisplayClusterLightCardTemplateList::Construct(const FArguments& InArgs, T
 			[
 				SAssignNew(LightCardTemplateTreeView, STreeView<TSharedPtr<FLightCardTemplateTreeItem>>)
 				.TreeItemsSource(&FilteredLightCardTemplateTree)
-				.ItemHeight(28)
 				.SelectionMode(ESelectionMode::Single)
 				.OnGenerateRow(this, &SDisplayClusterLightCardTemplateList::GenerateTreeItemRow)
 				.OnGetChildren(this, &SDisplayClusterLightCardTemplateList::GetChildrenForTreeItem)
diff --git a/Engine/Plugins/VirtualProduction/CameraCalibration/Source/CameraCalibrationEditor/Private/AssetEditor/SCameraCalibrationLinkedPointsDialog.cpp b/Engine/Plugins/VirtualProduction/CameraCalibration/Source/CameraCalibrationEditor/Private/AssetEditor/SCameraCalibrationLinkedPointsDialog.cpp
--- a/Engine/Plugins/VirtualProduction/CameraCalibration/Source/CameraCalibrationEditor/Private/AssetEditor/SCameraCalibrationLinkedPointsDialog.cpp
+++ b/Engine/Plugins/VirtualProduction/CameraCalibration/Source/CameraCalibrationEditor/Private/AssetEditor/SCameraCalibrationLinkedPointsDialog.cpp
@@ -191,7 +191,6 @@ void SCameraCalibrationLinkedPointsDialog::Construct(const FArguments& InArgs ,
         [
         	SAssignNew(LinkedItemsTree, STreeView<TSharedPtr<FLinkedTreeItem>>)
 			.TreeItemsSource(&LinkedItems)
-			.ItemHeight(24.0f)
 			.OnGenerateRow_Lambda([](TSharedPtr<FLinkedTreeItem> InItem, const TSharedRef<STableViewBase>& OwnerTable) { return InItem->MakeTreeRowWidget(OwnerTable); })
 			.OnGetChildren_Lambda([](TSharedPtr<FLinkedTreeItem> InItem, TArray<TSharedPtr<FLinkedTreeItem>>& OutNodes) { OutNodes = InItem->GetChildren(); })
 			.ClearSelectionOnClick(false)
diff --git a/Engine/Plugins/VirtualProduction/CameraCalibration/Source/CameraCalibrationEditor/Private/AssetEditor/SLensDataViewer.cpp b/Engine/Plugins/VirtualProduction/CameraCalibration/Source/CameraCalibrationEditor/Private/AssetEditor/SLensDataViewer.cpp
--- a/Engine/Plugins/VirtualProduction/CameraCalibration/Source/CameraCalibrationEditor/Private/AssetEditor/SLensDataViewer.cpp
+++ b/Engine/Plugins/VirtualProduction/CameraCalibration/Source/CameraCalibrationEditor/Private/AssetEditor/SLensDataViewer.cpp
@@ -307,7 +307,6 @@ TSharedRef<SWidget> SLensDataViewer::MakeLensDataWidget()
 			[
 				SAssignNew(TreeView, STreeView<TSharedPtr<FLensDataCategoryItem>>)
 				.TreeItemsSource(&DataCategories)
-				.ItemHeight(24.0f)
 				.OnGenerateRow(this, &SLensDataViewer::OnGenerateDataCategoryRow)
 				.OnGetChildren(this, &SLensDataViewer::OnGetDataCategoryItemChildren)
 				.OnSelectionChanged(this, &SLensDataViewer::OnDataCategorySelectionChanged)
@@ -345,7 +344,6 @@ TSharedRef<SWidget> SLensDataViewer::MakeLensDataWidget()
 			[
 					SAssignNew(DataEntriesTree, STreeView<TSharedPtr<FLensDataListItem>>)
 					.TreeItemsSource(&DataEntries)
-					.ItemHeight(24.0f)
 					.OnGenerateRow(this, &SLensDataViewer::OnGenerateDataEntryRow)
 					.OnGetChildren(this, &SLensDataViewer::OnGetDataEntryChildren)
 					.OnSelectionChanged(this, &SLensDataViewer::OnDataEntrySelectionChanged)
diff --git a/Engine/Plugins/VirtualProduction/CameraCalibration/Source/CameraCalibrationEditor/Private/Calibrators/CameraLensDistortionAlgoAruco.cpp b/Engine/Plugins/VirtualProduction/CameraCalibration/Source/CameraCalibrationEditor/Private/Calibrators/CameraLensDistortionAlgoAruco.cpp
--- a/Engine/Plugins/VirtualProduction/CameraCalibration/Source/CameraCalibrationEditor/Private/Calibrators/CameraLensDistortionAlgoAruco.cpp
+++ b/Engine/Plugins/VirtualProduction/CameraCalibration/Source/CameraCalibrationEditor/Private/Calibrators/CameraLensDistortionAlgoAruco.cpp
@@ -867,7 +867,6 @@ TSharedRef<SWidget> UCameraLensDistortionAlgoAruco::BuildCalibrationActionButton
 TSharedRef<SWidget> UCameraLensDistortionAlgoAruco::BuildCalibrationRowListWidget()
 {
 	CalibrationRowListWidget = SNew(SListView<TSharedPtr<FLensDistortionArucoRowData>>)
-		.ItemHeight(24)
 		.ListItemsSource(&CalibrationRows)
 		.SelectionMode(ESelectionMode::Multi)
 		.OnGenerateRow_Lambda([](TSharedPtr<FLensDistortionArucoRowData> InItem, const TSharedRef<STableViewBase>& OwnerTable) -> TSharedRef<ITableRow>
diff --git a/Engine/Plugins/VirtualProduction/CameraCalibration/Source/CameraCalibrationEditor/Private/Calibrators/CameraLensDistortionAlgoCheckerboard.cpp b/Engine/Plugins/VirtualProduction/CameraCalibration/Source/CameraCalibrationEditor/Private/Calibrators/CameraLensDistortionAlgoCheckerboard.cpp
--- a/Engine/Plugins/VirtualProduction/CameraCalibration/Source/CameraCalibrationEditor/Private/Calibrators/CameraLensDistortionAlgoCheckerboard.cpp
+++ b/Engine/Plugins/VirtualProduction/CameraCalibration/Source/CameraCalibrationEditor/Private/Calibrators/CameraLensDistortionAlgoCheckerboard.cpp
@@ -1008,7 +1008,6 @@ TSharedRef<SWidget> UCameraLensDistortionAlgoCheckerboard::BuildCalibrationActio
 TSharedRef<SWidget> UCameraLensDistortionAlgoCheckerboard::BuildCalibrationPointsTable()
 {
 	CalibrationListView = SNew(SListView<TSharedPtr<FLensDistortionCheckerboardRowData>>)
-		.ItemHeight(24)
 		.ListItemsSource(&CalibrationRows)
 		.OnGenerateRow_Lambda([&](TSharedPtr<FLensDistortionCheckerboardRowData> InItem, const TSharedRef<STableViewBase>& OwnerTable) -> TSharedRef<ITableRow>
 		{
diff --git a/Engine/Plugins/VirtualProduction/CameraCalibration/Source/CameraCalibrationEditor/Private/Calibrators/CameraLensDistortionAlgoPoints.cpp b/Engine/Plugins/VirtualProduction/CameraCalibration/Source/CameraCalibrationEditor/Private/Calibrators/CameraLensDistortionAlgoPoints.cpp
--- a/Engine/Plugins/VirtualProduction/CameraCalibration/Source/CameraCalibrationEditor/Private/Calibrators/CameraLensDistortionAlgoPoints.cpp
+++ b/Engine/Plugins/VirtualProduction/CameraCalibration/Source/CameraCalibrationEditor/Private/Calibrators/CameraLensDistortionAlgoPoints.cpp
@@ -897,7 +897,6 @@ TSharedRef<SWidget> UCameraLensDistortionAlgoPoints::BuildCurrentFocalLengthWidg
 TSharedRef<SWidget> UCameraLensDistortionAlgoPoints::BuildCalibrationPointsTable()
 {
 	CalibrationListView = SNew(SListView<TSharedPtr<FLensDistortionPointsRowData>>)
-		.ItemHeight(24)
 		.ListItemsSource(&CalibrationRows)
 		.OnGenerateRow_Lambda([&](TSharedPtr<FLensDistortionPointsRowData> InItem, const TSharedRef<STableViewBase>& OwnerTable) -> TSharedRef<ITableRow>
 		{
diff --git a/Engine/Plugins/VirtualProduction/CameraCalibration/Source/CameraCalibrationEditor/Private/Calibrators/CameraNodalOffsetAlgoPoints.cpp b/Engine/Plugins/VirtualProduction/CameraCalibration/Source/CameraCalibrationEditor/Private/Calibrators/CameraNodalOffsetAlgoPoints.cpp
--- a/Engine/Plugins/VirtualProduction/CameraCalibration/Source/CameraCalibrationEditor/Private/Calibrators/CameraNodalOffsetAlgoPoints.cpp
+++ b/Engine/Plugins/VirtualProduction/CameraCalibration/Source/CameraCalibrationEditor/Private/Calibrators/CameraNodalOffsetAlgoPoints.cpp
@@ -1935,7 +1935,6 @@ TSharedRef<SWidget> UCameraNodalOffsetAlgoPoints::BuildCalibrationPointsComboBox
 TSharedRef<SWidget> UCameraNodalOffsetAlgoPoints::BuildCalibrationPointsTable()
 {
 	CalibrationListView = SNew(SListView<TSharedPtr<FNodalOffsetPointsRowData>>)
-		.ItemHeight(24)
 		.ListItemsSource(&CalibrationRows)
 		.OnGenerateRow_Lambda([&](TSharedPtr<FNodalOffsetPointsRowData> InItem, const TSharedRef<STableViewBase>& OwnerTable) -> TSharedRef<ITableRow>
 		{
diff --git a/Engine/Plugins/VirtualProduction/DMX/DMXEngine/Source/DMXEditor/Private/Widgets/FixturePatch/SDMXMVRFixtureList.cpp b/Engine/Plugins/VirtualProduction/DMX/DMXEngine/Source/DMXEditor/Private/Widgets/FixturePatch/SDMXMVRFixtureList.cpp
--- a/Engine/Plugins/VirtualProduction/DMX/DMXEngine/Source/DMXEditor/Private/Widgets/FixturePatch/SDMXMVRFixtureList.cpp
+++ b/Engine/Plugins/VirtualProduction/DMX/DMXEngine/Source/DMXEditor/Private/Widgets/FixturePatch/SDMXMVRFixtureList.cpp
@@ -594,7 +594,6 @@ void SDMXMVRFixtureList::Construct(const FArguments& InArgs, TWeakPtr<FDMXEditor
 			SAssignNew(ListView, FDMXMVRFixtureListType)
 			.ListViewStyle(&TableViewStyle)
 			.HeaderRow(GenerateHeaderRow())
-			.ItemHeight(40.0f)
 			.ListItemsSource(&ListSource)
 			.OnGenerateRow(this, &SDMXMVRFixtureList::OnGenerateRow)
 			.OnSelectionChanged(this, &SDMXMVRFixtureList::OnSelectionChanged)
diff --git a/Engine/Plugins/VirtualProduction/DMX/DMXPixelMapping/Source/DMXPixelMappingEditor/Private/Views/SDMXPixelMappingHierarchyView.cpp b/Engine/Plugins/VirtualProduction/DMX/DMXPixelMapping/Source/DMXPixelMappingEditor/Private/Views/SDMXPixelMappingHierarchyView.cpp
--- a/Engine/Plugins/VirtualProduction/DMX/DMXPixelMapping/Source/DMXPixelMappingEditor/Private/Views/SDMXPixelMappingHierarchyView.cpp
+++ b/Engine/Plugins/VirtualProduction/DMX/DMXPixelMapping/Source/DMXPixelMappingEditor/Private/Views/SDMXPixelMappingHierarchyView.cpp
@@ -191,7 +191,6 @@ void SDMXPixelMappingHierarchyView::BuildChildSlotAndRefresh()
 			.Padding(4.f)
 			[
 				SAssignNew(HierarchyTreeView, STreeView<TSharedPtr<FDMXPixelMappingHierarchyItem>>)
-				.ItemHeight(20.0f)
 				.SelectionMode(ESelectionMode::Multi)
 				.HeaderRow(GenerateHeaderRow())
 				.OnGetChildren(FilterHandler.ToSharedRef(), &TreeFilterHandler<TSharedPtr<FDMXPixelMappingHierarchyItem>>::OnGetFilteredChildren)
diff --git a/Engine/Plugins/VirtualProduction/RemoteControl/Source/RemoteControlUI/Private/UI/SRCPanelExposedEntitiesList.cpp b/Engine/Plugins/VirtualProduction/RemoteControl/Source/RemoteControlUI/Private/UI/SRCPanelExposedEntitiesList.cpp
--- a/Engine/Plugins/VirtualProduction/RemoteControl/Source/RemoteControlUI/Private/UI/SRCPanelExposedEntitiesList.cpp
+++ b/Engine/Plugins/VirtualProduction/RemoteControl/Source/RemoteControlUI/Private/UI/SRCPanelExposedEntitiesList.cpp
@@ -292,7 +292,6 @@ void SRCPanelExposedEntitiesList::Construct(const FArguments& InArgs, URemoteCon
 
 	// Groups List
 	SAssignNew(GroupsListView, SListView<TSharedPtr<SRCPanelTreeNode>>)
-		.ItemHeight(24.f)
 		.SelectionMode(ESelectionMode::Single)
 		.ListItemsSource(reinterpret_cast<TArray<TSharedPtr<SRCPanelTreeNode>>*>(&FieldGroups))
 		.ClearSelectionOnClick(true)
@@ -335,7 +334,6 @@ void SRCPanelExposedEntitiesList::Construct(const FArguments& InArgs, URemoteCon
 
 	// Fields List
 	SAssignNew(FieldsListView, STreeView<TSharedPtr<SRCPanelTreeNode>>)
-		.ItemHeight(24.f)
 		.OnGenerateRow(this, &SRCPanelExposedEntitiesList::OnGenerateRow)
 		.OnSelectionChanged(this, &SRCPanelExposedEntitiesList::OnSelectionChanged)
 		.SelectionMode(ESelectionMode::Multi)
diff --git a/Engine/Plugins/VirtualProduction/RemoteControl/Source/RemoteControlUI/Private/UI/SSearchableTreeView.h b/Engine/Plugins/VirtualProduction/RemoteControl/Source/RemoteControlUI/Private/UI/SSearchableTreeView.h
--- a/Engine/Plugins/VirtualProduction/RemoteControl/Source/RemoteControlUI/Private/UI/SSearchableTreeView.h
+++ b/Engine/Plugins/VirtualProduction/RemoteControl/Source/RemoteControlUI/Private/UI/SSearchableTreeView.h
@@ -47,7 +47,6 @@ public:
 		IsSelectable = InArgs._IsSelectable;
 		
 		TreeView = SNew(STreeView<ItemType>)
-			.ItemHeight(24)
 			.TreeItemsSource(&FilteredItems)
 			.OnGenerateRow(this, &SSearchableTreeView::OnGenerateRow)
 			.OnGetChildren(this, &SSearchableTreeView::OnGetChildren)
