CameraCalibration: Improved media source selection for the lens file editor.
Users can now select from any media source or media texture asset in addition to media profile sources, as well as select live media tracks in the open level sequence editor.

diff --git a/Engine/Plugins/VirtualProduction/CameraCalibration/Source/CameraCalibrationEditor/CameraCalibrationEditor.Build.cs b/Engine/Plugins/VirtualProduction/CameraCalibration/Source/CameraCalibrationEditor/CameraCalibrationEditor.Build.cs
--- a/Engine/Plugins/VirtualProduction/CameraCalibration/Source/CameraCalibrationEditor/CameraCalibrationEditor.Build.cs
+++ b/Engine/Plugins/VirtualProduction/CameraCalibration/Source/CameraCalibrationEditor/CameraCalibrationEditor.Build.cs
@@ -18,6 +18,7 @@ namespace UnrealBuildTool.Rules
 					"CinematicCamera",
 					"Composure",
 					"ComposureLayersEditor",
+					"ContentBrowserData",
 					"Core",
 					"CoreUObject",
 	                "CurveEditor",
diff --git a/Engine/Plugins/VirtualProduction/CameraCalibration/Source/CameraCalibrationEditor/Private/AssetEditor/CameraCalibrationStepsController.cpp b/Engine/Plugins/VirtualProduction/CameraCalibration/Source/CameraCalibrationEditor/Private/AssetEditor/CameraCalibrationStepsController.cpp
--- a/Engine/Plugins/VirtualProduction/CameraCalibration/Source/CameraCalibrationEditor/Private/AssetEditor/CameraCalibrationStepsController.cpp
+++ b/Engine/Plugins/VirtualProduction/CameraCalibration/Source/CameraCalibrationEditor/Private/AssetEditor/CameraCalibrationStepsController.cpp
@@ -3,7 +3,6 @@
 #include "CameraCalibrationStepsController.h"
 
 #include "AssetRegistry/AssetData.h"
-#include "CalibrationPointComponent.h"
 #include "Camera/CameraActor.h"
 #include "CameraCalibrationEditorLog.h"
 #include "CameraCalibrationSettings.h"
@@ -16,6 +15,7 @@
 #include "Components/SceneCaptureComponent2D.h"
 #include "CompositingCaptureBase.h"
 #include "CompositingElement.h"
+#include "ContentBrowserDataSubsystem.h"
 #include "CompositingElements/CompositingElementInputs.h"
 #include "CompositingElements/CompositingElementOutputs.h"
 #include "CompositingElements/CompositingElementPasses.h"
@@ -26,6 +26,7 @@
 #include "Engine/UserDefinedEnum.h"
 #include "EngineUtils.h"
 #include "ICompElementManager.h"
+#include "IContentBrowserDataModule.h"
 #include "Input/Events.h"
 #include "Kismet/KismetRenderingLibrary.h"
 #include "LensComponent.h"
@@ -33,10 +34,10 @@
 #include "Materials/Material.h"
 #include "Materials/MaterialInstanceDynamic.h"
 #include "MediaPlayer.h"
+#include "MediaPlaylist.h"
 #include "MediaSource.h"
 #include "MediaTexture.h"
 #include "Misc/MessageDialog.h"
-#include "Models/SphericalLensModel.h"
 #include "Modules/ModuleManager.h"
 #include "Profile/IMediaProfileManager.h"
 #include "Profile/MediaProfile.h"
@@ -67,6 +68,40 @@ namespace CameraCalibrationStepsController
 		}
 		return nullptr;
 	}
+
+	/** Gets all assets of type TObject from the content browser */
+	template<typename TObject>
+	TArray<TSoftObjectPtr<TObject>> GetUObjectAssets()
+	{
+		TArray<TSoftObjectPtr<TObject>> Assets;
+	
+		UContentBrowserDataSubsystem* ContentBrowserData = IContentBrowserDataModule::Get().GetSubsystem();
+
+		FContentBrowserDataFilter Filter;
+		Filter.ItemTypeFilter = EContentBrowserItemTypeFilter::IncludeFiles;
+		Filter.ItemCategoryFilter = EContentBrowserItemCategoryFilter::IncludeAssets | EContentBrowserItemCategoryFilter::IncludeClasses | EContentBrowserItemCategoryFilter::IncludeCollections;
+		Filter.ItemAttributeFilter = EContentBrowserItemAttributeFilter::IncludeProject;
+		Filter.bRecursivePaths = true;
+	
+		FContentBrowserDataClassFilter& ClassFilter = Filter.ExtraFilters.FindOrAddFilter<FContentBrowserDataClassFilter>();
+		ClassFilter.ClassNamesToInclude.Add(TObject::StaticClass()->GetClassPathName().ToString());
+		ClassFilter.bRecursiveClassNamesToInclude = true;
+		ClassFilter.bRecursiveClassNamesToExclude = false;
+
+		ContentBrowserData->EnumerateItemsUnderPath(TEXT("/"), Filter, [&Assets](FContentBrowserItemData&& InItemData)
+		{
+			TSoftObjectPtr<TObject> MediaSourcePtr(FSoftObjectPath(InItemData.GetInternalPath().ToString()));
+
+			if (MediaSourcePtr.IsValid() || MediaSourcePtr.IsPending())
+			{
+				Assets.Add(MediaSourcePtr);
+			}
+		
+			return true;
+		});
+
+		return Assets;
+	}
 }
 
 
@@ -1526,60 +1561,95 @@ const ULensDistortionModelHandlerBase* FCameraCalibrationStepsController::GetDis
 	return nullptr;
 }
 
-bool FCameraCalibrationStepsController::SetMediaSourceUrl(const FString& InMediaSourceUrl)
+bool FCameraCalibrationStepsController::SetMediaSource(UMediaSource* InMediaSource)
 {
-	const UMediaProfile* MediaProfile = IMediaProfileManager::Get().GetCurrentMediaProfile();
-
-	if (!MediaProfile || !MediaPlayer.IsValid())
+	ClearMedia();
+	
+	if (!MediaPlayer.IsValid())
 	{
 		return false;
 	}
 
-	// If we're already playing it, we're done
-	if (InMediaSourceUrl == GetMediaSourceUrl())
+	if (!InMediaSource)
 	{
 		return true;
 	}
+	
+	MediaPlayer->OpenSource(InMediaSource);
+	MediaPlayer->Play();
+	return false;
+}
 
-	if (InMediaSourceUrl == TEXT("None") || !InMediaSourceUrl.Len())
+bool FCameraCalibrationStepsController::SetMediaTexture(UMediaTexture* InMediaTexture)
+{
+	ClearMedia();
+	
+	if (!MediaPlate.IsValid())
 	{
-		MediaPlayer->Close();
-		return true;
+		return false;
 	}
+	
+	ExternalMediaTexture = InMediaTexture;
 
-	for (int32 MediaSourceIdx = 0; MediaSourceIdx < MediaProfile->NumMediaSources(); ++MediaSourceIdx)
+	UMediaTexture* Texture = nullptr;
+	if (ExternalMediaTexture.IsValid())
+	{
+		Texture = ExternalMediaTexture.Get();
+	}
+	else
 	{
-		UMediaSource* MediaSource = MediaProfile->GetMediaSource(MediaSourceIdx);
+		Texture = MediaTexture.Get();
+	}
 
-		if (!MediaSource || (MediaSource->GetUrl() != InMediaSourceUrl))
+	if (Texture)
+	{
+		for (UCompositingElementInput* Input : MediaPlate->GetInputsList())
 		{
-			continue;
-		}
+			UMediaTextureCompositingInput* MediaInput = Cast<UMediaTextureCompositingInput>(Input);
 
-		MediaPlayer->OpenSource(MediaSource);
-		MediaPlayer->Play();
+			if (!MediaInput)
+			{
+				continue;
+			}
 
-		return true;
+			MediaInput->MediaSource = Texture;
+			break;
+		}
 	}
-
-	return false;
+	
+	return true;
 }
 
-/** Gets the current media source url being played. Empty if None */
-FString FCameraCalibrationStepsController::GetMediaSourceUrl() const
+void FCameraCalibrationStepsController::ClearMedia()
 {
-	if (!MediaPlayer.IsValid())
+	if (!MediaPlate.IsValid())
 	{
-		return TEXT("");
+		return;
 	}
 
-	return MediaPlayer->GetUrl();
+	MediaPlayer->Close();
+
+	ExternalMediaTexture = nullptr;
+	if (MediaTexture.IsValid())
+	{
+		for (UCompositingElementInput* Input : MediaPlate->GetInputsList())
+		{
+			UMediaTextureCompositingInput* MediaInput = Cast<UMediaTextureCompositingInput>(Input);
+
+			if (!MediaInput)
+			{
+				continue;
+			}
+
+			MediaInput->MediaSource = MediaTexture.Get();
+			break;
+		}
+	}
 }
 
-void FCameraCalibrationStepsController::FindMediaSourceUrls(TArray<TSharedPtr<FString>>& OutMediaSourceUrls) const
+void FCameraCalibrationStepsController::GetMediaProfileSources(TArray<TWeakObjectPtr<UMediaSource>>& OutSources) const
 {
 	const UMediaProfile* MediaProfile = IMediaProfileManager::Get().GetCurrentMediaProfile();
-
 	if (!MediaProfile)
 	{
 		return;
@@ -1587,13 +1657,53 @@ void FCameraCalibrationStepsController::FindMediaSourceUrls(TArray<TSharedPtr<FS
 
 	for (int32 MediaSourceIdx = 0; MediaSourceIdx < MediaProfile->NumMediaSources(); ++MediaSourceIdx)
 	{
-		if (const UMediaSource* MediaSource = MediaProfile->GetMediaSource(MediaSourceIdx))
+		if (UMediaSource* MediaSource = MediaProfile->GetMediaSource(MediaSourceIdx))
 		{
-			OutMediaSourceUrls.Add(MakeShared<FString>(MediaSource->GetUrl()));
+			// Media player will only play valid media sources, so only return the valid media sources from the profile
+			if (MediaSource->Validate())
+			{
+				OutSources.Add(TWeakObjectPtr<UMediaSource>(MediaSource));
+			}
 		}
 	}
 }
 
+TArray<TSoftObjectPtr<UMediaSource>> FCameraCalibrationStepsController::GetMediaSourceAssets() const
+{
+	return CameraCalibrationStepsController::GetUObjectAssets<UMediaSource>();
+}
+
+TArray<TSoftObjectPtr<UMediaTexture>> FCameraCalibrationStepsController::GetMediaTextureAssets() const
+{
+	return CameraCalibrationStepsController::GetUObjectAssets<UMediaTexture>();
+}
+
+/** Gets the current media source url being played. Empty if None */
+FString FCameraCalibrationStepsController::GetMediaSourceUrl() const
+{
+	if (!MediaPlayer.IsValid())
+	{
+		return TEXT("");
+	}
+
+	return MediaPlayer->GetUrl();
+}
+
+UMediaSource* FCameraCalibrationStepsController::GetMediaSource() const
+{
+	if (!MediaPlayer.IsValid())
+	{
+		return nullptr;
+	}
+
+	return MediaPlayer->GetPlaylist()->Get(MediaPlayer->GetPlaylistIndex());
+}
+
+UMediaTexture* FCameraCalibrationStepsController::GetMediaTexture() const
+{
+	return ExternalMediaTexture.Get();
+}
+
 const TConstArrayView<TStrongObjectPtr<UCameraCalibrationStep>> FCameraCalibrationStepsController::GetCalibrationSteps() const
 {
 	return TConstArrayView<TStrongObjectPtr<UCameraCalibrationStep>>(CalibrationSteps);
diff --git a/Engine/Plugins/VirtualProduction/CameraCalibration/Source/CameraCalibrationEditor/Private/AssetEditor/CameraCalibrationStepsController.h b/Engine/Plugins/VirtualProduction/CameraCalibration/Source/CameraCalibrationEditor/Private/AssetEditor/CameraCalibrationStepsController.h
--- a/Engine/Plugins/VirtualProduction/CameraCalibration/Source/CameraCalibrationEditor/Private/AssetEditor/CameraCalibrationStepsController.h
+++ b/Engine/Plugins/VirtualProduction/CameraCalibration/Source/CameraCalibrationEditor/Private/AssetEditor/CameraCalibrationStepsController.h
@@ -17,6 +17,7 @@ class ACompositingElement;
 class FCameraCalibrationToolkit;
 class SWidget;
 class UMediaPlayer;
+class UMediaSource;
 class UMediaTexture;
 class UCameraCalibrationStep;
 class UCompositingElementMaterialPass;
@@ -114,15 +115,33 @@ public:
 	/** Returns the distortion handler used to distort the CG being displayed in the simulcam viewport */
 	const ULensDistortionModelHandlerBase* GetDistortionHandler() const;
 
-	/** Sets the media source url to be played. Returns true if the url is a valid media source */
-	bool SetMediaSourceUrl(const FString& InMediaSourceUrl);
+	/** Sets the current media source being used in the calibration */
+	bool SetMediaSource(UMediaSource* InMediaSource);
 
-	/** Finds available media sources and adds their urls to the given array */
-	void FindMediaSourceUrls(TArray<TSharedPtr<FString>>& OutMediaSourceUrls) const;
+	/** Sets the current media to a media texture being used in the calibration */
+	bool SetMediaTexture(UMediaTexture* InMediaTexture);
 
+	/** Clears out any media source or media texture currently being played */
+	void ClearMedia();
+
+	/** Gets a list of valid media sources from the current media profile */
+	void GetMediaProfileSources(TArray<TWeakObjectPtr<UMediaSource>>& OutSources) const;
+
+	/** Gets all media source assets in the content browser */
+	TArray<TSoftObjectPtr<UMediaSource>> GetMediaSourceAssets() const;
+
+	/** Gets all media texture assets in the content browser */
+	TArray<TSoftObjectPtr<UMediaTexture>> GetMediaTextureAssets() const;
+	
 	/** Gets the current media source url being played. Empty if None */
 	FString GetMediaSourceUrl() const;
 
+	/** Gets the current media source being played, or nullptr if nothing is being played */
+	UMediaSource* GetMediaSource() const;
+
+	/** Gets the current media texture being displayed, or nullptr if the media texture is not being overridden */
+	UMediaTexture* GetMediaTexture() const;
+	
 	/** Returns the calibration steps */
 	const TConstArrayView<TStrongObjectPtr<UCameraCalibrationStep>> GetCalibrationSteps() const;
 
@@ -288,6 +307,9 @@ private:
 	/** The media texture used by the media plate */
 	TWeakObjectPtr<UMediaTexture> MediaTexture;
 
+	/** An external media texture to use by the media plate instead of the internal one */
+	TWeakObjectPtr<UMediaTexture> ExternalMediaTexture;
+	
 	/** The media player that is playing the selected media source */
 	TStrongObjectPtr<UMediaPlayer> MediaPlayer;
 
diff --git a/Engine/Plugins/VirtualProduction/CameraCalibration/Source/CameraCalibrationEditor/Private/AssetEditor/SCameraCalibrationSteps.cpp b/Engine/Plugins/VirtualProduction/CameraCalibration/Source/CameraCalibrationEditor/Private/AssetEditor/SCameraCalibrationSteps.cpp
--- a/Engine/Plugins/VirtualProduction/CameraCalibration/Source/CameraCalibrationEditor/Private/AssetEditor/SCameraCalibrationSteps.cpp
+++ b/Engine/Plugins/VirtualProduction/CameraCalibration/Source/CameraCalibrationEditor/Private/AssetEditor/SCameraCalibrationSteps.cpp
@@ -14,9 +14,12 @@
 #include "Engine/TextureRenderTarget2D.h"
 #include "IContentBrowserSingleton.h"
 #include "LensFile.h"
+#include "MediaSource.h"
+#include "MediaTexture.h"
 #include "PropertyCustomizationHelpers.h"
 #include "SSimulcamViewport.h"
 #include "Styling/CoreStyle.h"
+#include "Profile/MediaProfile.h"
 #include "Styling/SlateStyle.h"
 #include "UI/CameraCalibrationEditorStyle.h"
 #include "UI/CameraCalibrationWidgetHelpers.h"
@@ -316,8 +319,8 @@ void SCameraCalibrationSteps::Construct(const FArguments& InArgs, TWeakPtr<FCame
 				.MaxHeight(FCameraCalibrationWidgetHelpers::DefaultRowHeight)
 				[ FCameraCalibrationWidgetHelpers::BuildLabelWidgetPair(LOCTEXT("Camera", "Camera"), BuildCameraPickerWidget())]
 				
-				+ SVerticalBox::Slot() // Media Source picker
-				.MaxHeight(FCameraCalibrationWidgetHelpers::DefaultRowHeight)
+				+ SVerticalBox::Slot()
+				.AutoHeight()// Media Source picker
 				[FCameraCalibrationWidgetHelpers::BuildLabelWidgetPair(LOCTEXT("MediaSource", "Media Source"), BuildMediaSourceWidget())]
 
 				+ SVerticalBox::Slot() // Overlay picker
@@ -428,77 +431,226 @@ TSharedRef<SWidget> SCameraCalibrationSteps::BuildSimulcamWiperWidget()
 		.Delta(0.01f);
 }
 
-void SCameraCalibrationSteps::UpdateMediaSourcesOptions()
+TSharedRef<SWidget> SCameraCalibrationSteps::BuildMediaSourceTypeWidget()
 {
-	CurrentMediaSources.Empty();
+	static TArray<TSharedPtr<EMediaSourceType>> MediaSourceTypes =
+	{
+		MakeShared<EMediaSourceType>(EMediaSourceType::MediaSource),
+		MakeShared<EMediaSourceType>(EMediaSourceType::MediaTexture),
+		MakeShared<EMediaSourceType>(EMediaSourceType::MediaProfile),
+		MakeShared<EMediaSourceType>(EMediaSourceType::None),
+	};
 
-	if (CalibrationStepsController.IsValid())
+	static auto ToDisplayText = [](EMediaSourceType Type)
 	{
-		CalibrationStepsController.Pin()->FindMediaSourceUrls(CurrentMediaSources);
-	}
+		switch (Type)
+		{
+		case EMediaSourceType::MediaSource: return LOCTEXT("MediaSourceTypeLabel", "Media Source");
+		case EMediaSourceType::MediaTexture: return LOCTEXT("MediaTextureTypeLabel", "Media Texture");
+		case EMediaSourceType::MediaProfile: return LOCTEXT("MediaProfileTypeLabel", "Media Profile");
+		case EMediaSourceType::None: return LOCTEXT("MediaNoneTypeLabel", "None");
+		default: return FText::GetEmpty();
+		}
+	};
+	
+	return SNew(SComboBox<TSharedPtr<EMediaSourceType>>)
+	.OptionsSource(&MediaSourceTypes)
+	.OnGenerateWidget_Lambda([](TSharedPtr<EMediaSourceType> Type)
+	{
+		return SNew(STextBlock).Text(ToDisplayText(*Type));
+	})
+	.OnSelectionChanged_Lambda([this](TSharedPtr<EMediaSourceType> SelectedType, ESelectInfo::Type SelectType)
+	{
+		MediaSourceType = *SelectedType;
 
-	// Add a "None" option
-	CurrentMediaSources.Add(MakeShared<FString>(TEXT("None")));
+		// Clear out any selected media sources
+		if (MediaSourceType == EMediaSourceType::None)
+		{
+			if (CalibrationStepsController.IsValid())
+			{
+				CalibrationStepsController.Pin()->ClearMedia();
+				MediaProfileSourcesComboBox->SetSelectedItem(nullptr);
+			}
+		}
+	})
+	[
+		SNew(STextBlock)
+		.Text_Lambda([this]()
+		{
+			return ToDisplayText(MediaSourceType);
+		})
+	];
+}
 
-	check(MediaSourcesComboBox.IsValid());
+TSharedRef<SWidget> SCameraCalibrationSteps::BuildMediaSourceWidget()
+{
+	return SNew(SVerticalBox)
+
+	+SVerticalBox::Slot()
+	.AutoHeight()
+	.Padding(0.0, 2.0)
+	[
+		BuildMediaSourceTypeWidget()
+	]
+
+	+SVerticalBox::Slot()
+	.AutoHeight()
+	.Padding(0.0, 2.0)
+	[
+		SNew(SWidgetSwitcher)
+		.Visibility_Lambda([this]() { return MediaSourceType != EMediaSourceType::None ? EVisibility::Visible : EVisibility::Collapsed; })
+		.WidgetIndex_Lambda([this]() { return (int32)MediaSourceType; })
+
+		+SWidgetSwitcher::Slot()
+		[
+			BuildMediaSourceAssetPicker()
+		]
 
-	// Ask the ComboBox to refresh its options from its source (that we just updated)
-	MediaSourcesComboBox->RefreshOptions();
+		+SWidgetSwitcher::Slot()
+		[
+			BuildMediaTextureAssetPicker()
+		]
 
-	// Make sure we show the item that is selected
-	const FString MediaSourceUrl = CalibrationStepsController.Pin()->GetMediaSourceUrl();
+		+SWidgetSwitcher::Slot()
+		[
+			BuildMediaProfileSourcePicker()
+		]
+	];
+}
 
-	for (const TSharedPtr<FString>& MediaSourceUrlItem: CurrentMediaSources)
-	{
-		check(MediaSourceUrlItem.IsValid());
+TSharedRef<SWidget> SCameraCalibrationSteps::BuildMediaSourceAssetPicker()
+{
+	return SNew(SObjectPropertyEntryBox)
+		.AllowedClass(UMediaSource::StaticClass())
+		.ThumbnailPool(UThumbnailManager::Get().GetSharedThumbnailPool())
+		.AllowCreate(false)
+		.AllowClear(true)
+		.ObjectPath_Lambda([this]()
+		{
+			if (!CalibrationStepsController.IsValid())
+			{
+				return FString();
+			}
+
+			UMediaSource* MediaSource = CalibrationStepsController.Pin()->GetMediaSource();
+			if (!MediaSource)
+			{
+				return FString();
+			}
 
-		if (*MediaSourceUrlItem == MediaSourceUrl)
+			return MediaSource->GetPathName();
+		})
+		.OnObjectChanged_Lambda([this](const FAssetData& InAssetData)
 		{
-			MediaSourcesComboBox->SetSelectedItem(MediaSourceUrlItem);
-			return;
-		}
-	}
+			if (!CalibrationStepsController.IsValid())
+			{
+				return;
+			}
 
-	// If we arrived here, we fall back to "None"
-	MediaSourcesComboBox->SetSelectedItem(CurrentMediaSources[CurrentMediaSources.Num() - 1]);
+			CalibrationStepsController.Pin()->SetMediaSource(Cast<UMediaSource>(InAssetData.GetAsset()));
+			MediaProfileSourcesComboBox->SetSelectedItem(nullptr);
+		});
+}
 
-	return;
+TSharedRef<SWidget> SCameraCalibrationSteps::BuildMediaTextureAssetPicker()
+{
+	return SNew(SObjectPropertyEntryBox)
+		.AllowedClass(UMediaTexture::StaticClass())
+		.ThumbnailPool(UThumbnailManager::Get().GetSharedThumbnailPool())
+		.AllowCreate(false)
+		.AllowClear(true)
+		.ObjectPath_Lambda([this]()
+		{
+			if (!CalibrationStepsController.IsValid())
+			{
+				return FString();
+			}
+
+			UMediaTexture* MediaTexture = CalibrationStepsController.Pin()->GetMediaTexture();
+			if (!MediaTexture)
+			{
+				return FString();
+			}
+
+			return MediaTexture->GetPathName();
+		})
+		.OnObjectChanged_Lambda([this](const FAssetData& InAssetData)
+		{
+			if (!CalibrationStepsController.IsValid())
+			{
+				return;
+			}
+			
+			CalibrationStepsController.Pin()->SetMediaTexture(Cast<UMediaTexture>(InAssetData.GetAsset()));
+			MediaProfileSourcesComboBox->SetSelectedItem(nullptr);
+		});
 }
 
-TSharedRef<SWidget> SCameraCalibrationSteps::BuildMediaSourceWidget()
+TSharedRef<SWidget> SCameraCalibrationSteps::BuildMediaProfileSourcePicker()
 {
-	MediaSourcesComboBox = SNew(SComboBox<TSharedPtr<FString>>)
-		.OptionsSource(&CurrentMediaSources)
-		.OnSelectionChanged_Lambda([&](TSharedPtr<FString> NewValue, ESelectInfo::Type Type) -> void
+	return SAssignNew(MediaProfileSourcesComboBox, SComboBox<TWeakObjectPtr<UMediaSource>>)
+		.OptionsSource(&MediaProfileSources)
+		.OnComboBoxOpening_Lambda([this]()
 		{
+			MediaProfileSources.Empty();
+
+			if (!CalibrationStepsController.IsValid())
+			{
+				MediaProfileSourcesComboBox->RefreshOptions();
+				return;
+			}
+
+			CalibrationStepsController.Pin()->GetMediaProfileSources(MediaProfileSources);
+			MediaProfileSourcesComboBox->RefreshOptions();
+		})
+		.OnSelectionChanged_Lambda([&](TWeakObjectPtr<UMediaSource> NewValue, ESelectInfo::Type Type) -> void
+		{
+			if (Type == ESelectInfo::Direct)
+			{
+				return;
+			}
+			
 			if (!CalibrationStepsController.IsValid() || !NewValue.IsValid())
 			{
 				return;
 			}
 
-			CalibrationStepsController.Pin()->SetMediaSourceUrl(*NewValue);
+			CalibrationStepsController.Pin()->SetMediaSource(NewValue.Get());
 		})
-		.OnGenerateWidget_Lambda([&](TSharedPtr<FString> InOption) -> TSharedRef<SWidget>
+		.OnGenerateWidget_Lambda([&](TWeakObjectPtr<UMediaSource> InMediaSource) -> TSharedRef<SWidget>
 		{
-			return SNew(STextBlock).Text(FText::FromString(*InOption));
+			FText MediaSourceUrl = FText::GetEmpty();
+			if (InMediaSource.IsValid())
+			{
+				MediaSourceUrl = FText::FromString(InMediaSource->GetUrl());
+			}
+			
+			return SNew(STextBlock).Text(MediaSourceUrl);
 		})
-		.InitiallySelectedItem(nullptr)
 		[
 			SNew(STextBlock)
 			.Text_Lambda([&]() -> FText
 			{
-				if (MediaSourcesComboBox.IsValid() && MediaSourcesComboBox->GetSelectedItem().IsValid())
+				const FText None = LOCTEXT("NoneComboOption", "None");
+				if (!CalibrationStepsController.IsValid())
 				{
-					return FText::FromString(*MediaSourcesComboBox->GetSelectedItem());
+					return None;
 				}
 
-				return LOCTEXT("InvalidComboOption", "Invalid");
-			})
-		];
+				UMediaSource* MediaSource = CalibrationStepsController.Pin()->GetMediaSource();
+				if (!MediaSource)
+				{
+					return None;
+				}
 
-	UpdateMediaSourcesOptions();
+				if (!MediaSource->GetOuter()->IsA<UMediaProfile>())
+				{
+					return None;
+				}
 
-	return MediaSourcesComboBox.ToSharedRef();
+				return FText::FromString(MediaSource->GetUrl());
+			})
+		];
 }
 
 TSharedRef<SWidget> SCameraCalibrationSteps::BuildOverlayWidget()
diff --git a/Engine/Plugins/VirtualProduction/CameraCalibration/Source/CameraCalibrationEditor/Private/AssetEditor/SCameraCalibrationSteps.h b/Engine/Plugins/VirtualProduction/CameraCalibration/Source/CameraCalibrationEditor/Private/AssetEditor/SCameraCalibrationSteps.h
--- a/Engine/Plugins/VirtualProduction/CameraCalibration/Source/CameraCalibrationEditor/Private/AssetEditor/SCameraCalibrationSteps.h
+++ b/Engine/Plugins/VirtualProduction/CameraCalibration/Source/CameraCalibrationEditor/Private/AssetEditor/SCameraCalibrationSteps.h
@@ -26,6 +26,15 @@ struct FArguments;
  */
 class SCameraCalibrationSteps : public SCompoundWidget, public FGCObject
 {
+	/** Type indicating the kind of media source the user can pick from */
+	enum class EMediaSourceType : uint8
+	{
+		MediaSource,
+		MediaTexture,
+		MediaProfile,
+		None
+	};
+	
 	SLATE_BEGIN_ARGS(SCameraCalibrationSteps) {}
 	SLATE_END_ARGS()
 
@@ -49,9 +58,21 @@ private:
 	/** Builds the UI for the simulcam wiper */
 	TSharedRef<SWidget> BuildSimulcamWiperWidget();
 
+	/** Builds the media source type picker */
+	TSharedRef<SWidget> BuildMediaSourceTypeWidget();
+	
 	/** Builds the UI for the media source picker */
 	TSharedRef<SWidget> BuildMediaSourceWidget();
 
+	/** Builds an asset picker for media source assets */
+	TSharedRef<SWidget> BuildMediaSourceAssetPicker();
+
+	/** Builds an asset picker for media texture assets */
+	TSharedRef<SWidget> BuildMediaTextureAssetPicker();
+
+	/** Builds a combo box that displays the media source in the current media profile */
+	TSharedRef<SWidget> BuildMediaProfileSourcePicker();
+	
 	/** Builds the UI for the overlay picker */
 	TSharedRef<SWidget> BuildOverlayWidget();
 
@@ -61,9 +82,6 @@ private:
 	/** Updates the material parameter widget to display the parameters for the currently selected overlay */
 	void UpdateOverlayMaterialParameterWidget();
 
-	/** Refreshes the list of available media sources shown in the MediaSourcesComboBox */
-	void UpdateMediaSourcesOptions();
-
 	/** Expected to be called when user selects a new step via the UI */
 	void SelectStep(const FName& StepName);
 
@@ -74,12 +92,15 @@ private:
 
 	/** The controller object */
 	TWeakPtr<class FCameraCalibrationStepsController> CalibrationStepsController;
+
+	/** The currently selected media source type */
+	EMediaSourceType MediaSourceType = EMediaSourceType::MediaProfile;
 	
 	/** Options source for the MediaSourcesComboBox. Lists the currently available media sources */
-	TArray<TSharedPtr<FString>> CurrentMediaSources;
+	TArray<TWeakObjectPtr<UMediaSource>> MediaProfileSources;
 
 	/** The combobox that presents the available media sources */
-	TSharedPtr<SComboBox<TSharedPtr<FString>>> MediaSourcesComboBox;
+	TSharedPtr<SComboBox<TWeakObjectPtr<UMediaSource>>> MediaProfileSourcesComboBox;
 
 	/** The combox that presents the available overlays */
 	TSharedPtr<SComboBox<TSharedPtr<FName>>> OverlayComboBox;
