CameraCalibration: Add Target Offset parameters to the lens solver to account for errors in tracking props with calibration patterns.

diff --git a/Engine/Plugins/Runtime/OpenCV/Source/OpenCVHelper/Private/OpenCVHelper.cpp b/Engine/Plugins/Runtime/OpenCV/Source/OpenCVHelper/Private/OpenCVHelper.cpp
--- a/Engine/Plugins/Runtime/OpenCV/Source/OpenCVHelper/Private/OpenCVHelper.cpp
+++ b/Engine/Plugins/Runtime/OpenCV/Source/OpenCVHelper/Private/OpenCVHelper.cpp
@@ -404,10 +404,10 @@ bool FOpenCVHelper::IdentifyArucoMarkers(TArray<FColor>& Image, FIntPoint ImageS
 
 		const std::vector<cv::Point2f>& MarkerCorners = Corners[MarkerIndex];
 
-		NewMarker.Corners[0] = FVector2f(MarkerCorners[0].x, MarkerCorners[0].y);
-		NewMarker.Corners[1] = FVector2f(MarkerCorners[1].x, MarkerCorners[1].y);
-		NewMarker.Corners[2] = FVector2f(MarkerCorners[2].x, MarkerCorners[2].y);
-		NewMarker.Corners[3] = FVector2f(MarkerCorners[3].x, MarkerCorners[3].y);
+		NewMarker.Corners[0] = FVector2f(MarkerCorners[0].x, MarkerCorners[0].y); // TopLeft
+		NewMarker.Corners[1] = FVector2f(MarkerCorners[1].x, MarkerCorners[1].y); // TopRight
+		NewMarker.Corners[2] = FVector2f(MarkerCorners[2].x, MarkerCorners[2].y); // BottomRight
+		NewMarker.Corners[3] = FVector2f(MarkerCorners[3].x, MarkerCorners[3].y); // BottomLeft
 
 		OutMarkers.Add(NewMarker);
 	}
diff --git a/Engine/Plugins/VirtualProduction/CameraCalibration/Source/CameraCalibrationEditor/Private/Calibrators/CameraCalibrationSolver.cpp b/Engine/Plugins/VirtualProduction/CameraCalibration/Source/CameraCalibrationEditor/Private/Calibrators/CameraCalibrationSolver.cpp
--- a/Engine/Plugins/VirtualProduction/CameraCalibration/Source/CameraCalibrationEditor/Private/Calibrators/CameraCalibrationSolver.cpp
+++ b/Engine/Plugins/VirtualProduction/CameraCalibration/Source/CameraCalibrationEditor/Private/Calibrators/CameraCalibrationSolver.cpp
@@ -2,6 +2,7 @@
 
 #include "CameraCalibrationSolver.h"
 
+#include "CameraCalibrationUtils.h"
 #include "HAL/IConsoleManager.h"
 #include "Models/AnamorphicLensModel.h"
 #include "Models/SphericalLensModel.h"
@@ -37,7 +38,7 @@ bool ULensDistortionSolver::GetStatusText(FText& OutStatusText)
 
 void ULensDistortionSolver::SetStatusText(FText InStatusText) 
 { 
-	StatusText = MoveTemp(InStatusText);
+	StatusText = InStatusText;
 	bHasStatusChanged = true;
 }
 
@@ -54,6 +55,7 @@ FDistortionCalibrationResult ULensDistortionSolverOpenCV::Solve_Implementation(
 	const FVector2D& ImageCenter,
 	const TArray<float>& DistortionParameters,
 	const TArray<FTransform>& CameraPoses,
+	const TArray<FTransform>& TargetPoses,
 	TSubclassOf<ULensModel> LensModel,
 	double PixelAspect,
 	ECalibrationFlags SolverFlags)
@@ -84,6 +86,11 @@ FDistortionCalibrationResult ULensDistortionSolverOpenCV::Solve_Implementation(
 	// Reorganize the 3D and 2D points from the input arrays or arrays to be laid out linearly in memory in two cv::Mat objects
 	GatherPoints(ObjectPointArray, ImagePointArray, ObjectPointsMat, ImagePointsMat);
 
+	// Find the set of unique camera poses to reduce the number of poses we need to solve for
+	TArray<int32> CameraPoseIndices;
+	TArray<FTransform> UniqueCameraPoses;
+	FindUniqueCameraPoses(CameraPoses, UniqueCameraPoses, CameraPoseIndices);
+
 	double RMSE = 0.0;
 
 	cv::Mat CameraMatrix = cv::Mat::eye(3, 3, CV_64F);
@@ -239,9 +246,15 @@ FDistortionCalibrationResult ULensDistortionSolverOpenCV::Solve_Implementation(
 	}
 
 	// Initialize the solver with the number of parameters to solve, and the maximum number of iterations to run
-	const int NumIntrinsics = NumDistortionCoefficients + 4; // Includes Fx, Fy, Cx, and Cy
+	int NumIntrinsics = NumDistortionCoefficients + 4; // Includes Fx, Fy, Cx, and Cy
 	const int NumExtrinsics = 6; // 3 for rotation vector, 3 for translation vector
 
+	if ((EnumHasAnyFlags(SolverFlags, ECalibrationFlags::SolveTargetOffset)))
+	{
+		const int NumObjectOffsetParams = 6; // 3 for rotation vector, 3 for translation vector
+		NumIntrinsics += NumObjectOffsetParams;
+	}
+
 	// If using a guess for the extrinsic parameters, then the solver will be constrained to solve only one camera pose, and needs only one set of extrinsic parameters.
 	// Otherwise, the solver needs one set of extrinsic parameters per image.
 	const bool bConstrainWithCameraMotion = CVarConstrainWithCameraMotion.GetValueOnAnyThread();
@@ -253,7 +266,7 @@ FDistortionCalibrationResult ULensDistortionSolverOpenCV::Solve_Implementation(
 	}
 	else
 	{
-		NumPosesToSolve = NumImages;
+		NumPosesToSolve = UniqueCameraPoses.Num();
 	}
 
 	const int NumParamsToSolve = NumIntrinsics + (NumPosesToSolve * NumExtrinsics);
@@ -270,7 +283,20 @@ FDistortionCalibrationResult ULensDistortionSolverOpenCV::Solve_Implementation(
 	SolverParams[2] = CameraMatrix.at<double>(0, 2);
 	SolverParams[3] = CameraMatrix.at<double>(1, 2);
 
-	FMemory::Memcpy(SolverParams + 4, DistCoeffs.ptr<double>(), NumDistortionCoefficients * sizeof(double));
+	// Initialize the target offset rvec and tvec to zero
+	if ((EnumHasAnyFlags(SolverFlags, ECalibrationFlags::SolveTargetOffset)))
+	{
+		SolverParams[4] = 0.0f;
+		SolverParams[5] = 0.0f;
+		SolverParams[6] = 0.0f;
+		SolverParams[7] = 0.0f;
+		SolverParams[8] = 0.0f;
+		SolverParams[9] = 0.0f;
+	}
+
+	// Initialize the solver with the starting values for the distortion coefficients
+	const int32 DistortionParameterOffset = NumIntrinsics - NumDistortionCoefficients;
+	FMemory::Memcpy(SolverParams + DistortionParameterOffset, DistCoeffs.ptr<double>(), NumDistortionCoefficients * sizeof(double));
 
 	// Instruct the solver to ignore some parameters when running its solve
 	uchar* Mask = Solver.Mask.ptr<uchar>();
@@ -297,15 +323,15 @@ FDistortionCalibrationResult ULensDistortionSolverOpenCV::Solve_Implementation(
 	}
 	if (EnumHasAnyFlags(SolverFlags, ECalibrationFlags::FixDistortion))
 	{
-		for (int ParamIndex = 0; ParamIndex < NumDistortionCoefficients; ++ParamIndex)
+		for (int ParamIndex = DistortionParameterOffset; ParamIndex < DistortionParameterOffset + NumDistortionCoefficients; ++ParamIndex)
 		{
-			Mask[ParamIndex + 4] = 0;
+			Mask[ParamIndex] = 0;
 		}
 	}
 
 	if (LensModel == UAnamorphicLensModel::StaticClass())
 	{
-		Mask[4] = 0; // We do not want to solve for pixel aspect
+		Mask[DistortionParameterOffset] = 0; // We do not want to solve for pixel aspect
 	}
 
 	// Initialize the starting guess for the camera's extrinsic parameters. 
@@ -321,7 +347,7 @@ FDistortionCalibrationResult ULensDistortionSolverOpenCV::Solve_Implementation(
 
 		if ((EnumHasAnyFlags(SolverFlags, ECalibrationFlags::UseExtrinsicGuess)))
 		{
-			FOpenCVHelper::MakeObjectVectorsFromCameraPose(CameraPoses[ImageIndex], Rotation, Translation);
+			FOpenCVHelper::MakeObjectVectorsFromCameraPose(UniqueCameraPoses[ImageIndex], Rotation, Translation);
 		}
 		else
 		{
@@ -375,7 +401,15 @@ FDistortionCalibrationResult ULensDistortionSolverOpenCV::Solve_Implementation(
 		CameraMatrix.at<double>(0, 2) = SolverParams[2];
 		CameraMatrix.at<double>(1, 2) = SolverParams[3];
 
-		FMemory::Memcpy(DistCoeffs.ptr<double>(), SolverParams + 4, NumDistortionCoefficients * sizeof(double));
+		cv::Mat TargetOffsetRotation = cv::Mat::zeros(3, 1, CV_64F);;
+		cv::Mat TargetOffsetTranslation = cv::Mat::zeros(3, 1, CV_64F);;
+		if (EnumHasAnyFlags(SolverFlags, ECalibrationFlags::SolveTargetOffset))
+		{
+			TargetOffsetRotation = Solver.Params.rowRange(4, 7);
+			TargetOffsetTranslation = Solver.Params.rowRange(7, 10);
+		}
+
+		FMemory::Memcpy(DistCoeffs.ptr<double>(), SolverParams + DistortionParameterOffset, NumDistortionCoefficients * sizeof(double));
 
 		// If the solver determined that it no longer needs to proceed, break out of the loop
 		if (!bShouldProceed)
@@ -404,6 +438,15 @@ FDistortionCalibrationResult ULensDistortionSolverOpenCV::Solve_Implementation(
 			cv::Mat ObjectPointsInImage = cv::Mat(ObjectPointsMat.colRange(ObjectPointsIndex, ObjectPointsIndex + NumImagePoints));
 			cv::Mat ImagePointsInImage = cv::Mat(ImagePointsMat.colRange(ObjectPointsIndex, ObjectPointsIndex + NumImagePoints));
 
+			FTransform TargetPose = FTransform::Identity;
+			if (EnumHasAnyFlags(SolverFlags, ECalibrationFlags::SolveTargetOffset))
+			{
+				if (TargetPoses.IsValidIndex(ImageIndex))
+				{
+					TargetPose = TargetPoses[ImageIndex];
+				}
+			}
+
 			ObjectPointsIndex += NumImagePoints;
 
 			int ExtrinsicOffset = 0;
@@ -413,7 +456,7 @@ FDistortionCalibrationResult ULensDistortionSolverOpenCV::Solve_Implementation(
 			}
 			else
 			{
-				ExtrinsicOffset = NumIntrinsics + (ImageIndex * NumExtrinsics);
+				ExtrinsicOffset = NumIntrinsics + (CameraPoseIndices[ImageIndex] * NumExtrinsics);
 			}
 
 			cv::Mat Rotation;
@@ -439,11 +482,11 @@ FDistortionCalibrationResult ULensDistortionSolverOpenCV::Solve_Implementation(
 			// Project the 3D object points to 2D using the camera extrinsics, intrinsics, and distortion parameters, optionally also solving for the Jacobian matrix
 			if (bComputeJacobian)
 			{
-				ProjectPoints(LensModel, ObjectPointsInImage, Rotation, Translation, CameraMatrix, DistCoeffs, CvImageSize, ProjectedPoints, Jacobian, SolverFlags);
+				ProjectPoints(LensModel, ObjectPointsInImage, TargetPose, TargetOffsetRotation, TargetOffsetTranslation, Rotation, Translation, CameraMatrix, DistCoeffs, CvImageSize, ProjectedPoints, Jacobian, SolverFlags);
 			}
 			else
 			{
-				ProjectPoints(LensModel, ObjectPointsInImage, Rotation, Translation, CameraMatrix, DistCoeffs, CvImageSize, ProjectedPoints);
+				ProjectPoints(LensModel, ObjectPointsInImage, TargetPose, TargetOffsetRotation, TargetOffsetTranslation, Rotation, Translation, CameraMatrix, DistCoeffs, CvImageSize, ProjectedPoints);
 			}
 
 			// Compute the difference between the input image points and the projected points
@@ -489,7 +532,7 @@ FDistortionCalibrationResult ULensDistortionSolverOpenCV::Solve_Implementation(
 	for (int ImageIndex = 0; ImageIndex < NumImages; ++ImageIndex)
 	{
 		// Get the rotation and translation vectors for this image
-		int ExtrinsicOffset = NumIntrinsics + (ImageIndex * NumExtrinsics);
+		int ExtrinsicOffset = NumIntrinsics + (CameraPoseIndices[ImageIndex] * NumExtrinsics);
 		cv::Mat Rotation = Solver.Params.rowRange(ExtrinsicOffset, ExtrinsicOffset + 3);
 		cv::Mat Translation = Solver.Params.rowRange(ExtrinsicOffset + 3, ExtrinsicOffset + 6);
 
@@ -501,14 +544,14 @@ FDistortionCalibrationResult ULensDistortionSolverOpenCV::Solve_Implementation(
 	}
 
 	// If the solver computed exactly one calibrated camera pose, we can compute a tracking offset based on the input camera pose
-	if (Result.CameraPoses.Num() == 1 && CameraPoses.Num() == 1)
+	if (UniqueCameraPoses.Num() == 1)
 	{
 		const FTransform& TrackedCameraPose = CameraPoses[0];
 		const FTransform& CalibratedCameraPose = Result.CameraPoses[0];
-		const FTransform TrackingOffset = CalibratedCameraPose * TrackedCameraPose.Inverse();
+		const FTransform TargetOffset = CalibratedCameraPose * TrackedCameraPose.Inverse();
 
-		Result.NodalOffset.LocationOffset = TrackingOffset.GetLocation();
-		Result.NodalOffset.RotationOffset = TrackingOffset.GetRotation();
+		Result.NodalOffset.LocationOffset = TargetOffset.GetLocation();
+		Result.NodalOffset.RotationOffset = TargetOffset.GetRotation();
 	}
 
 	if (LensModel == UAnamorphicLensModel::StaticClass())
@@ -859,13 +902,18 @@ void ULensDistortionSolverOpenCV::InitCameraExtrinsics(
 
 		cv::Mat ProjectedPoints = Solver.Error.reshape(2, 1);
 
+		// When initializing the extrinsics, the target offset is not used
+		FTransform TargetPose = FTransform::Identity;
+		cv::Mat TargetOffsetRotation = cv::Mat::zeros(3, 1, CV_64F);
+		cv::Mat TargetOffsetTranslation = cv::Mat::zeros(3, 1, CV_64F);
+
 		if (bComputeJacobian)
 		{
-			ProjectPoints(LensModel, ObjectPoints, Rotation, Translation, CameraMatrix, DistCoeffs, ImageSize, ProjectedPoints, Solver.Jacobian, SolverFlags);
+			ProjectPoints(LensModel, ObjectPoints, TargetPose, TargetOffsetRotation, TargetOffsetTranslation, Rotation, Translation, CameraMatrix, DistCoeffs, ImageSize, ProjectedPoints, Solver.Jacobian, SolverFlags);
 		}
 		else
 		{
-			ProjectPoints(LensModel, ObjectPoints, Rotation, Translation, CameraMatrix, DistCoeffs, ImageSize, ProjectedPoints);
+			ProjectPoints(LensModel, ObjectPoints, TargetPose, TargetOffsetRotation, TargetOffsetTranslation, Rotation, Translation, CameraMatrix, DistCoeffs, ImageSize, ProjectedPoints);
 		}
 
 		ProjectedPoints = ProjectedPoints - ImagePoints;
@@ -875,6 +923,9 @@ void ULensDistortionSolverOpenCV::InitCameraExtrinsics(
 void ULensDistortionSolverOpenCV::ProjectPoints(
 	const TSubclassOf<ULensModel> LensModel,
 	const cv::Mat& ObjectPoints,
+	const FTransform& TargetPose,
+	const cv::Mat& TargetOffsetRotation,
+	const cv::Mat& TargetOffsetTranslation,
 	const cv::Mat& Rotation,
 	const cv::Mat& Translation,
 	const cv::Mat& CameraMatrix,
@@ -884,12 +935,15 @@ void ULensDistortionSolverOpenCV::ProjectPoints(
 {
 	cv::Mat Jacobian;
 	const ECalibrationFlags SolverFlags = ECalibrationFlags::None;
-	ProjectPoints(LensModel, ObjectPoints, Rotation, Translation, CameraMatrix, DistCoeffs, ImageSize, ProjectedPoints, Jacobian, SolverFlags);
+	ProjectPoints(LensModel, ObjectPoints, TargetPose, TargetOffsetRotation, TargetOffsetTranslation, Rotation, Translation, CameraMatrix, DistCoeffs, ImageSize, ProjectedPoints, Jacobian, SolverFlags);
 }
 
 void ULensDistortionSolverOpenCV::ProjectPoints(
 	const TSubclassOf<ULensModel> LensModel,
 	const cv::Mat& ObjectPoints,
+	const FTransform& TargetPose,
+	const cv::Mat& TargetOffsetRotation,
+	const cv::Mat& TargetOffsetTranslation,
 	const cv::Mat& Rotation,
 	const cv::Mat& Translation,
 	const cv::Mat& CameraMatrix,
@@ -904,6 +958,9 @@ void ULensDistortionSolverOpenCV::ProjectPoints(
 	{
 		ProjectPointsSpherical(
 			ObjectPoints,
+			TargetPose,
+			TargetOffsetRotation,
+			TargetOffsetTranslation,
 			Rotation,
 			Translation,
 			CameraMatrix,
@@ -1072,8 +1129,10 @@ void ULensDistortionSolverOpenCV::ProjectPointsAnamorphic(
 			cv::Mat JacTranslation;
 			cv::Mat JacFocalLength;
 			cv::Mat JacImageCenter;
+			cv::Mat JacTargetOffsetRotation;
+			cv::Mat JacTargetOffsetTranslation;
 			cv::Mat JacDistortion;
-			SubdivideJacobian(Jacobian, JacRotation, JacTranslation, JacFocalLength, JacImageCenter, JacDistortion, SolverFlags);
+			SubdivideJacobian(Jacobian, JacRotation, JacTranslation, JacFocalLength, JacImageCenter, JacTargetOffsetRotation, JacTargetOffsetTranslation, JacDistortion, SolverFlags);
 
 			if (!JacDistortion.empty())
 			{
@@ -1252,6 +1311,9 @@ void ULensDistortionSolverOpenCV::ProjectPointsAnamorphic(
 
 void ULensDistortionSolverOpenCV::ProjectPointsSpherical(
 	const cv::Mat& ObjectPoints,
+	const FTransform& TargetPose,
+	const cv::Mat& TargetOffsetRotation,
+	const cv::Mat& TargetOffsetTranslation,
 	const cv::Mat& Rotation,
 	const cv::Mat& Translation,
 	const cv::Mat& CameraMatrix,
@@ -1287,13 +1349,73 @@ void ULensDistortionSolverOpenCV::ProjectPointsSpherical(
 	const double P2 = DistortionParams[3];
 	const double K3 = DistortionParams[4];
 
+	// Convert the target offset rvec into rodrigues format
+	double TargetOffsetR[9];
+	cv::Mat TargetOffsetRotationMatrix = cv::Mat(3, 3, CV_64F, TargetOffsetR);
+
+	double TargetOffsetJacR[27];
+	cv::Mat TargetOffsetRotationJacobian = cv::Mat(3, 9, CV_64F, TargetOffsetJacR);
+
+	cv::Rodrigues(TargetOffsetRotation, TargetOffsetRotationMatrix, TargetOffsetRotationJacobian);
+
+	const double* TargetOffsetT = TargetOffsetTranslation.ptr<double>();
+
+	// Convert the target pose transform into rvecs and tvecs
+	cv::Mat TargetPoseRvec;
+	cv::Mat TargetPoseTvec;
+	FOpenCVHelper::MakeObjectVectorsFromCameraPose(TargetPose, TargetPoseRvec, TargetPoseTvec);
+
+	// Convert the target pose rvec into rodrigues format
+	double TrackerR[9];
+	cv::Mat TrackerRotationMatrix = cv::Mat(3, 3, CV_64F, TrackerR);
+
+	double TrackerJacR[27];
+	cv::Mat TrackerRotationJacobian = cv::Mat(3, 9, CV_64F, TrackerJacR);
+
+	cv::Rodrigues(TargetPoseRvec, TrackerRotationMatrix, TrackerRotationJacobian);
+
+	const double* TrackerT = TargetPoseTvec.ptr<double>();
+
+	// Convert the inverse target pose transform into rvecs and tvecs
+	cv::Mat InvTargetPoseRvec;
+	cv::Mat InvTargetPoseTvec;
+	FOpenCVHelper::MakeObjectVectorsFromCameraPose(TargetPose.Inverse(), InvTargetPoseRvec, InvTargetPoseTvec);
+
+	// Convert the inverse target pose rvec into rodrigues format
+	double InvTrackerR[9];
+	cv::Mat InvTrackerRotationMatrix = cv::Mat(3, 3, CV_64F, InvTrackerR);
+
+	double InvTrackerJacR[27];
+	cv::Mat InvTrackerRotationJacobian = cv::Mat(3, 9, CV_64F, InvTrackerJacR);
+
+	cv::Rodrigues(InvTargetPoseRvec, InvTrackerRotationMatrix, InvTrackerRotationJacobian);
+
+	const double* InvTrackerT = InvTargetPoseTvec.ptr<double>();
+
 	const int NumPoints = ObjectPoints.total();
 
 	for (int PointIndex = 0; PointIndex < NumPoints; PointIndex++)
 	{
-		const cv::Point3d& WorldPoint = ObjectPoints.at<cv::Point3d>(PointIndex);
-		cv::Point2d& ProjectedPoint = ProjectedPoints.at<cv::Point2d>(PointIndex);
-
+		const cv::Point3d& InputPoint = ObjectPoints.at<cv::Point3d>(PointIndex);
+
+		// Transform the input point by the original target pose
+		cv::Point3d TargetPoint;
+		TargetPoint.x = TrackerR[0] * InputPoint.x + TrackerR[1] * InputPoint.y + TrackerR[2] * InputPoint.z + TrackerT[0];
+		TargetPoint.y = TrackerR[3] * InputPoint.x + TrackerR[4] * InputPoint.y + TrackerR[5] * InputPoint.z + TrackerT[1];
+		TargetPoint.z = TrackerR[6] * InputPoint.x + TrackerR[7] * InputPoint.y + TrackerR[8] * InputPoint.z + TrackerT[2];
+
+		// Transform the input point by the target offset
+		cv::Point3d OffsetPoint;
+		OffsetPoint.x = TargetOffsetR[0] * TargetPoint.x + TargetOffsetR[1] * TargetPoint.y + TargetOffsetR[2] * TargetPoint.z + TargetOffsetT[0];
+		OffsetPoint.y = TargetOffsetR[3] * TargetPoint.x + TargetOffsetR[4] * TargetPoint.y + TargetOffsetR[5] * TargetPoint.z + TargetOffsetT[1];
+		OffsetPoint.z = TargetOffsetR[6] * TargetPoint.x + TargetOffsetR[7] * TargetPoint.y + TargetOffsetR[8] * TargetPoint.z + TargetOffsetT[2];
+
+		// Transform the input point by the inverse of the target pose
+		cv::Point3d WorldPoint;
+		WorldPoint.x = InvTrackerR[0] * OffsetPoint.x + InvTrackerR[1] * OffsetPoint.y + InvTrackerR[2] * OffsetPoint.z + InvTrackerT[0];
+		WorldPoint.y = InvTrackerR[3] * OffsetPoint.x + InvTrackerR[4] * OffsetPoint.y + InvTrackerR[5] * OffsetPoint.z + InvTrackerT[1];
+		WorldPoint.z = InvTrackerR[6] * OffsetPoint.x + InvTrackerR[7] * OffsetPoint.y + InvTrackerR[8] * OffsetPoint.z + InvTrackerT[2];
+		
 		cv::Point3d CameraPoint;
 		CameraPoint.x = R[0] * WorldPoint.x + R[1] * WorldPoint.y + R[2] * WorldPoint.z + TranslationVector[0];
 		CameraPoint.y = R[3] * WorldPoint.x + R[4] * WorldPoint.y + R[5] * WorldPoint.z + TranslationVector[1];
@@ -1321,6 +1443,7 @@ void ULensDistortionSolverOpenCV::ProjectPointsSpherical(
 		DistortedPoint.x = (CameraPoint.x * Radial) + TangentialX;
 		DistortedPoint.y = (CameraPoint.y * Radial) + TangentialY;
 
+		cv::Point2d& ProjectedPoint = ProjectedPoints.at<cv::Point2d>(PointIndex);
 		ProjectedPoint.x = DistortedPoint.x * Fx + Cx;
 		ProjectedPoint.y = DistortedPoint.y * Fy + Cy;
 
@@ -1332,8 +1455,10 @@ void ULensDistortionSolverOpenCV::ProjectPointsSpherical(
 			cv::Mat JacTranslation;
 			cv::Mat JacFocalLength;
 			cv::Mat JacImageCenter;
+			cv::Mat JacTargetOffsetRotation;
+			cv::Mat JacTargetOffsetTranslation;
 			cv::Mat JacDistortion;
-			SubdivideJacobian(Jacobian, JacRotation, JacTranslation, JacFocalLength, JacImageCenter, JacDistortion, SolverFlags);
+			SubdivideJacobian(Jacobian, JacRotation, JacTranslation, JacFocalLength, JacImageCenter, JacTargetOffsetRotation, JacTargetOffsetTranslation, JacDistortion, SolverFlags);
 
 			if (!JacImageCenter.empty())
 			{
@@ -1419,7 +1544,6 @@ void ULensDistortionSolverOpenCV::ProjectPointsSpherical(
 				JacTranslationPtr[JacobianStep + 0] = Fy * DDistortedPointYDT.x;
 				JacTranslationPtr[JacobianStep + 1] = Fy * DDistortedPointYDT.y;
 				JacTranslationPtr[JacobianStep + 2] = Fy * DDistortedPointYDT.z;
-
 			}
 
 			if (!JacRotation.empty())
@@ -1468,7 +1592,112 @@ void ULensDistortionSolverOpenCV::ProjectPointsSpherical(
 				JacRotationPtr[JacobianStep + 0] = Fy * DDistortedPointYDR.x;
 				JacRotationPtr[JacobianStep + 1] = Fy * DDistortedPointYDR.y;
 				JacRotationPtr[JacobianStep + 2] = Fy * DDistortedPointYDR.z;
+			}
+
+			if (!JacTargetOffsetTranslation.empty())
+			{
+				cv::Point3d DWorldPointXDT;
+				DWorldPointXDT.x = InvTrackerR[0];
+				DWorldPointXDT.y = InvTrackerR[1];
+				DWorldPointXDT.z = InvTrackerR[2];
+
+				cv::Point3d DWorldPointYDT;
+				DWorldPointYDT.x = InvTrackerR[3];
+				DWorldPointYDT.y = InvTrackerR[4];
+				DWorldPointYDT.z = InvTrackerR[5];
+
+				cv::Point3d DWorldPointZDT;
+				DWorldPointZDT.x = InvTrackerR[6];
+				DWorldPointZDT.y = InvTrackerR[7];
+				DWorldPointZDT.z = InvTrackerR[8];
+
+				cv::Point3d DCameraPointXDT = R[0] * DWorldPointXDT + R[1] * DWorldPointYDT + R[2] * DWorldPointZDT;;
+				cv::Point3d DCameraPointYDT = R[3] * DWorldPointXDT + R[4] * DWorldPointYDT + R[5] * DWorldPointZDT;
+				cv::Point3d DCameraPointZDT = R[6] * DWorldPointXDT + R[7] * DWorldPointYDT + R[8] * DWorldPointZDT;
+
+				// Compute the derivative of the undistorted image point with respect to the rotation vector [ R0  R1  R2 ]
+				const cv::Point3d DImagePointXDT = CameraPoint.z * (DCameraPointXDT - CameraPoint.x * DCameraPointZDT);
+				const cv::Point3d DImagePointYDT = CameraPoint.z * (DCameraPointYDT - CameraPoint.y * DCameraPointZDT);
+
+				// Compute the derivative of the R-Squared term with respect to the rotation vector [ R0  R1  R2 ]
+				const cv::Point3d DR2DT= 2 * ((CameraPoint.x * DImagePointXDT) + (CameraPoint.y * DImagePointYDT));
+
+				// Compute the derivative of the radial distortion term with respect to the rotation vector [ R0  R1  R2 ]
+				const cv::Point3d DRadialDT = DR2DT * (K1 + (2 * K2 * R2) + (3 * K3 * R4));
+
+				// Compute the derivative of the A1 term with respect to the rotation vector [ R0  R1  R2 ]
+				const cv::Point3d DA1DT = 2 * ((CameraPoint.x * DImagePointYDT) + (CameraPoint.y * DImagePointXDT));
+
+				// Compute the derivative of distorted point with respect to the rotation vector [ R0  R1  R2 ]
+				const cv::Point3d DDistortedPointXDT = (DImagePointXDT * Radial) + (CameraPoint.x * DRadialDT) + (P1 * DA1DT) + (P2 * (DR2DT + 4 * CameraPoint.x * DImagePointXDT));
+				const cv::Point3d DDistortedPointYDT = (DImagePointYDT * Radial) + (CameraPoint.y * DRadialDT) + (P2 * DA1DT) + (P1 * (DR2DT + 4 * CameraPoint.y * DImagePointYDT));
+
+				// Add derivatives to the Jacobian matrix
+				double* JacTargetOffsetTranslationPtr = JacTargetOffsetTranslation.ptr<double>(PointIndex * 2);
+
+				JacTargetOffsetTranslationPtr[0] = Fx * DDistortedPointXDT.x;
+				JacTargetOffsetTranslationPtr[1] = Fx * DDistortedPointXDT.y;
+				JacTargetOffsetTranslationPtr[2] = Fx * DDistortedPointXDT.z;
 
+				JacTargetOffsetTranslationPtr[JacobianStep + 0] = Fy * DDistortedPointYDT.x;
+				JacTargetOffsetTranslationPtr[JacobianStep + 1] = Fy * DDistortedPointYDT.y;
+				JacTargetOffsetTranslationPtr[JacobianStep + 2] = Fy * DDistortedPointYDT.z;
+			}
+
+			if (!JacTargetOffsetRotation.empty())
+			{
+				cv::Point3d DOffsetPointXDR;
+				DOffsetPointXDR.x = TargetPoint.x * TargetOffsetJacR[0] + TargetPoint.y * TargetOffsetJacR[1] + TargetPoint.z * TargetOffsetJacR[2];
+				DOffsetPointXDR.y = TargetPoint.x * TargetOffsetJacR[9] + TargetPoint.y * TargetOffsetJacR[10] + TargetPoint.z * TargetOffsetJacR[11];
+				DOffsetPointXDR.z = TargetPoint.x * TargetOffsetJacR[18] + TargetPoint.y * TargetOffsetJacR[19] + TargetPoint.z * TargetOffsetJacR[20];
+
+				cv::Point3d DOffsetPointYDR;
+				DOffsetPointYDR.x = TargetPoint.x * TargetOffsetJacR[3] + TargetPoint.y * TargetOffsetJacR[4] + TargetPoint.z * TargetOffsetJacR[5];
+				DOffsetPointYDR.y = TargetPoint.x * TargetOffsetJacR[12] + TargetPoint.y * TargetOffsetJacR[13] + TargetPoint.z * TargetOffsetJacR[14];
+				DOffsetPointYDR.z = TargetPoint.x * TargetOffsetJacR[21] + TargetPoint.y * TargetOffsetJacR[22] + TargetPoint.z * TargetOffsetJacR[23];
+
+				cv::Point3d DOffsetPointZDR;
+				DOffsetPointZDR.x = TargetPoint.x * TargetOffsetJacR[6] + TargetPoint.y * TargetOffsetJacR[7] + TargetPoint.z * TargetOffsetJacR[8];
+				DOffsetPointZDR.y = TargetPoint.x * TargetOffsetJacR[15] + TargetPoint.y * TargetOffsetJacR[16] + TargetPoint.z * TargetOffsetJacR[17];
+				DOffsetPointZDR.z = TargetPoint.x * TargetOffsetJacR[24] + TargetPoint.y * TargetOffsetJacR[25] + TargetPoint.z * TargetOffsetJacR[26];
+
+				cv::Point3d DWorldPointXDR = InvTrackerR[0] * DOffsetPointXDR + InvTrackerR[1] * DOffsetPointYDR + InvTrackerR[2] * DOffsetPointZDR;
+				cv::Point3d DWorldPointYDR = InvTrackerR[3] * DOffsetPointXDR + InvTrackerR[4] * DOffsetPointYDR + InvTrackerR[5] * DOffsetPointZDR;
+				cv::Point3d DWorldPointZDR = InvTrackerR[6] * DOffsetPointXDR + InvTrackerR[7] * DOffsetPointYDR + InvTrackerR[8] * DOffsetPointZDR;
+
+				// Compute the derivatives of the camera points with respect to the rotation vector [ R0  R1  R2]
+				// This requires the use of the rotation jacobian matrix that relates the rotation vector in Rodrigues form to the rotation matrix
+				cv::Point3d DCameraPointXDR = R[0] * DWorldPointXDR + R[1] * DWorldPointYDR + R[2] * DWorldPointZDR;
+				cv::Point3d DCameraPointYDR = R[3] * DWorldPointXDR + R[4] * DWorldPointYDR + R[5] * DWorldPointZDR;
+				cv::Point3d DCameraPointZDR = R[6] * DWorldPointXDR + R[7] * DWorldPointYDR + R[8] * DWorldPointZDR;
+
+				// Compute the derivative of the undistorted image point with respect to the rotation vector [ R0  R1  R2 ]
+				const cv::Point3d DImagePointXDR = CameraPoint.z * (DCameraPointXDR - CameraPoint.x * DCameraPointZDR);
+				const cv::Point3d DImagePointYDR = CameraPoint.z * (DCameraPointYDR - CameraPoint.y * DCameraPointZDR);
+
+				// Compute the derivative of the R-Squared term with respect to the rotation vector [ R0  R1  R2 ]
+				const cv::Point3d DR2DR = 2 * ((CameraPoint.x * DImagePointXDR) + (CameraPoint.y * DImagePointYDR));
+
+				// Compute the derivative of the radial distortion term with respect to the rotation vector [ R0  R1  R2 ]
+				const cv::Point3d DRadialDR = DR2DR * (K1 + (2 * K2 * R2) + (3 * K3 * R4));
+
+				// Compute the derivative of the A1 term with respect to the rotation vector [ R0  R1  R2 ]
+				const cv::Point3d DA1DR = 2 * ((CameraPoint.x * DImagePointYDR) + (CameraPoint.y * DImagePointXDR));
+
+				// Compute the derivative of distorted point with respect to the rotation vector [ R0  R1  R2 ]
+				const cv::Point3d DDistortedPointXDR = (DImagePointXDR * Radial) + (CameraPoint.x * DRadialDR) + (P1 * DA1DR) + (P2 * (DR2DR + 4 * CameraPoint.x * DImagePointXDR));
+				const cv::Point3d DDistortedPointYDR = (DImagePointYDR * Radial) + (CameraPoint.y * DRadialDR) + (P2 * DA1DR) + (P1 * (DR2DR + 4 * CameraPoint.y * DImagePointYDR));
+
+				// Add derivatives to the Jacobian matrix
+				double* JacTargetOffsetRotationPtr = JacTargetOffsetRotation.ptr<double>(PointIndex * 2);
+
+				JacTargetOffsetRotationPtr[0] = Fx * DDistortedPointXDR.x;
+				JacTargetOffsetRotationPtr[1] = Fx * DDistortedPointXDR.y;
+				JacTargetOffsetRotationPtr[2] = Fx * DDistortedPointXDR.z;
+
+				JacTargetOffsetRotationPtr[JacobianStep + 0] = Fy * DDistortedPointYDR.x;
+				JacTargetOffsetRotationPtr[JacobianStep + 1] = Fy * DDistortedPointYDR.y;
+				JacTargetOffsetRotationPtr[JacobianStep + 2] = Fy * DDistortedPointYDR.z;
 			}
 		}
 	}
@@ -1512,6 +1741,8 @@ void ULensDistortionSolverOpenCV::SubdivideJacobian(
 	cv::Mat& JacTranslation,
 	cv::Mat& JacFocalLength,
 	cv::Mat& JacImageCenter,
+	cv::Mat& JacTargetOffsetRotation,
+	cv::Mat& JacTargetOffsetTranslation,
 	cv::Mat& JacDistortion,
 	ECalibrationFlags SolverFlags)
 {
@@ -1541,11 +1772,45 @@ void ULensDistortionSolverOpenCV::SubdivideJacobian(
 			}
 		}
 
-		if (Jacobian.cols > 10)
+		int DistortionParameterOffset = 10;
+		if (EnumHasAnyFlags(SolverFlags, ECalibrationFlags::SolveTargetOffset))
 		{
-			JacDistortion = cv::Mat(Jacobian.colRange(10, Jacobian.cols));
+			if (Jacobian.cols >= 16)
+			{
+				JacTargetOffsetRotation = cv::Mat(Jacobian.colRange(10, 13));
+				JacTargetOffsetTranslation = cv::Mat(Jacobian.colRange(13, 16));
+			}
+
+			DistortionParameterOffset += 6;
 		}
+
+		if (Jacobian.cols > DistortionParameterOffset)
+		{
+			JacDistortion = cv::Mat(Jacobian.colRange(DistortionParameterOffset, Jacobian.cols));
+		}
+	}
+}
+
+void ULensDistortionSolverOpenCV::FindUniqueCameraPoses(const TArray<FTransform>& InCameraPoses, TArray<FTransform>& OutUniqueCameraPoses, TArray<int32>& OutUniquePoseIndices)
+{
+	const int32 NumPoses = InCameraPoses.Num();
+
+	FTransform LastCameraPose = InCameraPoses[0];
+	FTransform CurrentCameraPose;
+	for (int32 PoseIndex = 0; PoseIndex < NumPoses; ++PoseIndex)
+	{
+		CurrentCameraPose = InCameraPoses[PoseIndex];
+
+		if (!FCameraCalibrationUtils::IsNearlyEqual(LastCameraPose, CurrentCameraPose))
+		{
+			OutUniqueCameraPoses.Add(LastCameraPose);
+			LastCameraPose = CurrentCameraPose;
+		}
+
+		OutUniquePoseIndices.Add(OutUniqueCameraPoses.Num());
 	}
+
+	OutUniqueCameraPoses.Add(LastCameraPose);
 }
 
 FLevMarqSolver::FLevMarqSolver(int NumParamsToSolve, int NumErrors, int NumMaxIterations, ESymmetryMode SymmetryMode)
diff --git a/Engine/Plugins/VirtualProduction/CameraCalibration/Source/CameraCalibrationEditor/Private/Calibrators/CameraCalibrationSolver.h b/Engine/Plugins/VirtualProduction/CameraCalibration/Source/CameraCalibrationEditor/Private/Calibrators/CameraCalibrationSolver.h
--- a/Engine/Plugins/VirtualProduction/CameraCalibration/Source/CameraCalibrationEditor/Private/Calibrators/CameraCalibrationSolver.h
+++ b/Engine/Plugins/VirtualProduction/CameraCalibration/Source/CameraCalibrationEditor/Private/Calibrators/CameraCalibrationSolver.h
@@ -57,7 +57,8 @@ enum class ECalibrationFlags : uint32
 	FixPrincipalPoint = 1 << 3,  /** The solver will not optimize the principal point */
 	FixExtrinsics = 1 << 4,      /** The solver will not optimize the camera extrinsics */
 	FixDistortion = 1 << 5,      /** The solver will fix all distortion values at 0 */
-	FixAspectRatio = 1 << 6      /** The solver will respect the input aspect ratio when solving for Fx and Fy */
+	FixAspectRatio = 1 << 6,     /** The solver will respect the input aspect ratio when solving for Fx and Fy */
+	SolveTargetOffset = 1 << 7   /** The solver will solve for an offset for each image's input 3D object points */
 };
 
 /** 
@@ -79,6 +80,7 @@ public:
 		const FVector2D& ImageCenter,
 		const TArray<float>& DistortionParameters,
 		const TArray<FTransform>& CameraPoses,
+		const TArray<FTransform>& TargetPoses,
 		TSubclassOf<ULensModel> LensModel,
 		double PixelAspect,
 		ECalibrationFlags SolverFlags);
@@ -106,6 +108,7 @@ public:
 		const FVector2D& ImageCenter,
 		const TArray<float>& DistortionParameters,
 		const TArray<FTransform>& CameraPoses,
+		const TArray<FTransform>& TargetPoses,
 		TSubclassOf<ULensModel> LensModel,
 		double PixelAspect,
 		ECalibrationFlags SolverFlags) PURE_VIRTUAL(ULensDistortionSolver::Solve_Implementation, return FDistortionCalibrationResult(););
@@ -156,6 +159,7 @@ public:
 		const FVector2D& ImageCenter,
 		const TArray<float>& DistortionParameters,
 		const TArray<FTransform>& CameraPoses,
+		const TArray<FTransform>& TargetPoses,
 		TSubclassOf<ULensModel> LensModel,
 		double PixelAspect,
 		ECalibrationFlags SolverFlags) override;
@@ -188,6 +192,9 @@ private:
 	void ProjectPoints(
 		const TSubclassOf<ULensModel> LensModel,
 		const cv::Mat& ObjectPoints,
+		const FTransform& TargetPose,
+		const cv::Mat& ObjectPointOffsetRotation,
+		const cv::Mat& ObjectPointOffsetTranslation,
 		const cv::Mat& Rotation,
 		const cv::Mat& Translation,
 		const cv::Mat& CameraMatrix,
@@ -199,6 +206,9 @@ private:
 	void ProjectPoints(
 		const TSubclassOf<ULensModel> LensModel,
 		const cv::Mat& ObjectPoints,
+		const FTransform& TargetPose,
+		const cv::Mat& ObjectPointOffsetRotation,
+		const cv::Mat& ObjectPointOffsetTranslation,
 		const cv::Mat& Rotation,
 		const cv::Mat& Translation,
 		const cv::Mat& CameraMatrix,
@@ -223,6 +233,9 @@ private:
 	/** Project the input object points to 2D using the input camera intrinsics, extrinsics, and spherical distortion parameters */
 	void ProjectPointsSpherical(
 		const cv::Mat& ObjectPoints,
+		const FTransform& TargetPose,
+		const cv::Mat& ObjectPointOffsetRotation,
+		const cv::Mat& ObjectPointOffsetTranslation,
 		const cv::Mat& Rotation,
 		const cv::Mat& Translation,
 		const cv::Mat& CameraMatrix,
@@ -245,9 +258,14 @@ private:
 		cv::Mat& JacTranslation,
 		cv::Mat& JacFocalLength,
 		cv::Mat& JacImageCenter,
+		cv::Mat& JacObjectPointOffsetRotation,
+		cv::Mat& JacObjectPointOffsetTranslation,
 		cv::Mat& JacDistortion,
 		ECalibrationFlags SolverFlags);
 
+	/** Find the unique set of camera poses and the index into that set for each image */
+	void FindUniqueCameraPoses(const TArray<FTransform>& InCameraPoses, TArray<FTransform>& OutUniqueCameraPoses, TArray<int32>& OutUniquePoseIndices);
+
 #endif	// WITH_OPENCV
 };
 
diff --git a/Engine/Plugins/VirtualProduction/CameraCalibration/Source/CameraCalibrationEditor/Private/Calibrators/CameraLensDistortionAlgoAruco.cpp b/Engine/Plugins/VirtualProduction/CameraCalibration/Source/CameraCalibrationEditor/Private/Calibrators/CameraLensDistortionAlgoAruco.cpp
--- a/Engine/Plugins/VirtualProduction/CameraCalibration/Source/CameraCalibrationEditor/Private/Calibrators/CameraLensDistortionAlgoAruco.cpp
+++ b/Engine/Plugins/VirtualProduction/CameraCalibration/Source/CameraCalibrationEditor/Private/Calibrators/CameraLensDistortionAlgoAruco.cpp
@@ -38,7 +38,7 @@
 static TAutoConsoleVariable<bool> CVarFixExtrinsicsAruco(TEXT("LensDistortionAruco.FixExtrinsics"), false, TEXT("If true, the solver will fix the camera extrinsics to the user-provided camera poses"));
 static TAutoConsoleVariable<bool> CVarFixDistortionAruco(TEXT("LensDistortionAruco.FixDistortion"), false, TEXT("If true, the solver will not optimize distortion, and will use the input distortion values if any are given, or assume zero distortion otherwise."));
 static TAutoConsoleVariable<bool> CVarUseExtrinsicsGuessAruco(TEXT("LensDistortionAruco.UseExtrinsicsGuess"), false, TEXT("If true, the actual calibrator and camera poses will be used when running the solver"));
-static TAutoConsoleVariable<bool> CVarGroupPointsByCameraPoseAruco(TEXT("LensDistortionAruco.GroupPointsByCameraPose"), true, TEXT("If true, the points sent to the solver will be grouped together if they share the same camera pose."));
+static TAutoConsoleVariable<bool> CVarSolveTargetOffsetAruco(TEXT("LensDistortionAruco.SolveTargetOffset"), false, TEXT("If true, the solver will calibrate for an offset from the calibrator tracking origin that further minimizes reprojection error."));
 #endif
 
 const int UCameraLensDistortionAlgoAruco::DATASET_VERSION = 1;
@@ -340,6 +340,9 @@ FDistortionCalibrationTask UCameraLensDistortionAlgoAruco::BeginCalibration(FTex
 	TArray<FTransform> CameraPoses;
 	CameraPoses.Reserve(CalibrationRows.Num());
 
+	TArray<FTransform> TargetPoses;
+	CameraPoses.Reserve(CalibrationRows.Num());
+
 	// Extract the 3D points, 2D points, and camera poses from each row to pass to the solver
 	for (const TSharedPtr<FLensDistortionArucoRowData>& Row : CalibrationRows)
 	{
@@ -365,11 +368,21 @@ FDistortionCalibrationTask UCameraLensDistortionAlgoAruco::BeginCalibration(FTex
 		Samples3d.Add(Points3d);
 		Samples2d.Add(Points2d);
 		CameraPoses.Add(Row->CameraPose);
-	}
 
-	if (CVarGroupPointsByCameraPoseAruco.GetValueOnGameThread())
-	{
-		UE::CameraCalibration::Private::GroupPointsByCameraPose(Samples3d, Samples2d, CameraPoses);
+		FTransform FirstMarkerPose;
+
+		FArucoCalibrationPoint& FirstAruco = Row->ArucoPoints[0];
+
+		FVector TopLeft = FirstAruco.Corners3D[0];
+		FVector TopRight = FirstAruco.Corners3D[1];
+		FVector BottomLeft = FirstAruco.Corners3D[3];
+
+		FirstMarkerPose.SetLocation(TopLeft);
+
+		FRotator FirstMarkerRotation = FRotationMatrix::MakeFromYZ(TopRight - TopLeft, TopLeft - BottomLeft).Rotator();
+		FirstMarkerPose.SetRotation(FirstMarkerRotation.Quaternion());
+
+		TargetPoses.Add(FirstMarkerPose);
 	}
 
 	ECalibrationFlags SolverFlags = ECalibrationFlags::None;
@@ -404,14 +417,19 @@ FDistortionCalibrationTask UCameraLensDistortionAlgoAruco::BeginCalibration(FTex
 		EnumAddFlags(SolverFlags, ECalibrationFlags::FixPrincipalPoint);
 	}
 
+	if (CVarSolveTargetOffsetAruco.GetValueOnGameThread())
+	{
+		EnumAddFlags(SolverFlags, ECalibrationFlags::SolveTargetOffset);
+	}
+
 	FDistortionInfo DistortionGuess;
-	LensFile->GetDistortionPoint(Focus, Zoom, DistortionGuess);
+	LensFile->EvaluateDistortionParameters(Focus, Zoom, DistortionGuess);
 
 	const TSubclassOf<ULensModel> Model = LensFile->LensInfo.LensModel;
 
 	UClass* SolverClass = LensDistortionTool->GetSolverClass();
 
-	CalibrationTask = UE::Tasks::Launch(UE_SOURCE_LOCATION, [SolverClass, Model, Samples3d, Samples2d, ImageSize, FocalLength, ImageCenter, DistortionGuess, CameraPoses, PixelAspect, SolverFlags, Focus, Zoom]() mutable
+	CalibrationTask = UE::Tasks::Launch(UE_SOURCE_LOCATION, [SolverClass, Model, Samples3d, Samples2d, ImageSize, FocalLength, ImageCenter, DistortionGuess, CameraPoses, TargetPoses, PixelAspect, SolverFlags, Focus, Zoom]() mutable
 		{
 			ULensDistortionSolver* Solver = NewObject<ULensDistortionSolver>(GetTransientPackage(), SolverClass);
 
@@ -423,6 +441,7 @@ FDistortionCalibrationTask UCameraLensDistortionAlgoAruco::BeginCalibration(FTex
 				ImageCenter,
 				DistortionGuess.Parameters,
 				CameraPoses,
+				TargetPoses,
 				Model,
 				PixelAspect,
 				SolverFlags
diff --git a/Engine/Plugins/VirtualProduction/CameraCalibration/Source/CameraCalibrationEditor/Private/Calibrators/CameraLensDistortionAlgoCheckerboard.cpp b/Engine/Plugins/VirtualProduction/CameraCalibration/Source/CameraCalibrationEditor/Private/Calibrators/CameraLensDistortionAlgoCheckerboard.cpp
--- a/Engine/Plugins/VirtualProduction/CameraCalibration/Source/CameraCalibrationEditor/Private/Calibrators/CameraLensDistortionAlgoCheckerboard.cpp
+++ b/Engine/Plugins/VirtualProduction/CameraCalibration/Source/CameraCalibrationEditor/Private/Calibrators/CameraLensDistortionAlgoCheckerboard.cpp
@@ -42,7 +42,7 @@ static TAutoConsoleVariable<bool> CVarFixExtrinsics(TEXT("LensDistortionCheckerb
 static TAutoConsoleVariable<bool> CVarFixDistortion(TEXT("LensDistortionCheckerboard.FixDistortion"), false, TEXT("If true, the solver will not optimize distortion, and will use the input distortion values if any are given, or assume zero distortion otherwise."));
 static TAutoConsoleVariable<bool> CVarUseExtrinsicsGuess(TEXT("LensDistortionCheckerboard.UseExtrinsicsGuess"), false, TEXT("If true, the actual checkerboard and camera poses will be used when running the solver"));
 static TAutoConsoleVariable<bool> CVarUseTrackedCalibrator(TEXT("LensDistortionCheckerboard.UseTrackedCalibrator"), false, TEXT("If true, the 3D points of the checkerboard will come from the tracked pose of the object, otherwise, dummy points will be used."));
-static TAutoConsoleVariable<bool> CVarGroupPointsByCameraPose(TEXT("LensDistortionCheckerboard.GroupPointsByCameraPose"), true, TEXT("If true, the points sent to the solver will be grouped together if they share the same camera pose."));
+static TAutoConsoleVariable<bool> CVarSolveTargetOffset(TEXT("LensDistortionCheckerboard.SolveTargetOffset"), false, TEXT("If true, the solver will calibrate for an offset from the calibrator tracking origin that further minimizes reprojection error."));
 #endif
 
 const int UCameraLensDistortionAlgoCheckerboard::DATASET_VERSION = 1;
@@ -680,6 +680,9 @@ FDistortionCalibrationTask UCameraLensDistortionAlgoCheckerboard::BeginCalibrati
 	TArray<FTransform> CameraPoses;
 	CameraPoses.Reserve(CalibrationRows.Num());
 
+	TArray<FTransform> TargetPoses;
+	CameraPoses.Reserve(CalibrationRows.Num());
+
 	for (const TSharedPtr<FLensDistortionCheckerboardRowData>& Row : CalibrationRows)
 	{
 		FObjectPoints Points3d;
@@ -695,6 +698,19 @@ FDistortionCalibrationTask UCameraLensDistortionAlgoCheckerboard::BeginCalibrati
 		Samples3d.Add(Points3d);
 		Samples2d.Add(Points2d);
 		CameraPoses.Add(Row->CameraData.Pose);
+
+		FTransform BoardPose;
+
+		FVector TopLeft = Row->Points3d[0];
+		FVector TopRight = Row->Points3d[Row->NumCornerCols - 1];
+		FVector BottomLeft = Row->Points3d[Row->NumCornerCols * (Row->NumCornerRows - 1)];
+
+		BoardPose.SetLocation(TopLeft);
+
+		FRotator BoardRotation = FRotationMatrix::MakeFromYZ(TopRight - TopLeft, TopLeft - BottomLeft).Rotator();
+		BoardPose.SetRotation(BoardRotation.Quaternion());
+
+		TargetPoses.Add(BoardPose);
 	}
 
 	if (!CVarUseTrackedCalibrator.GetValueOnGameThread())
@@ -716,10 +732,6 @@ FDistortionCalibrationTask UCameraLensDistortionAlgoCheckerboard::BeginCalibrati
 			Samples3d.Add(Points3d);
 		}
 	}
-	else if (CVarGroupPointsByCameraPose.GetValueOnGameThread())
-	{
-		UE::CameraCalibration::Private::GroupPointsByCameraPose(Samples3d, Samples2d, CameraPoses);
-	}
 
 	ECalibrationFlags SolverFlags = ECalibrationFlags::None;
 
@@ -753,14 +765,19 @@ FDistortionCalibrationTask UCameraLensDistortionAlgoCheckerboard::BeginCalibrati
 		EnumAddFlags(SolverFlags, ECalibrationFlags::FixPrincipalPoint);
 	}
 
+	if (CVarSolveTargetOffset.GetValueOnGameThread())
+	{
+		EnumAddFlags(SolverFlags, ECalibrationFlags::SolveTargetOffset);
+	}
+
 	FDistortionInfo DistortionGuess;
-	LensFile->GetDistortionPoint(Focus, Zoom, DistortionGuess);
+	LensFile->EvaluateDistortionParameters(Focus, Zoom, DistortionGuess);
 
 	const TSubclassOf<ULensModel> Model = LensFile->LensInfo.LensModel;
 	
 	Solver = NewObject<ULensDistortionSolver>(this, Tool->GetSolverClass());
 
-	CalibrationTask = UE::Tasks::Launch(UE_SOURCE_LOCATION, [Solver = Solver, Model, Samples3d, Samples2d, ImageSize, FocalLength, ImageCenter, DistortionGuess, CameraPoses, PixelAspect, SolverFlags, Focus, Zoom]() mutable
+	CalibrationTask = UE::Tasks::Launch(UE_SOURCE_LOCATION, [Solver = Solver, Model, Samples3d, Samples2d, ImageSize, FocalLength, ImageCenter, DistortionGuess, CameraPoses, TargetPoses, PixelAspect, SolverFlags, Focus, Zoom]() mutable
 		{
 			FDistortionCalibrationResult Result = Solver->Solve(
 				Samples3d,
@@ -770,6 +787,7 @@ FDistortionCalibrationTask UCameraLensDistortionAlgoCheckerboard::BeginCalibrati
 				ImageCenter,
 				DistortionGuess.Parameters,
 				CameraPoses,
+				TargetPoses,
 				Model,
 				PixelAspect,
 				SolverFlags
diff --git a/Engine/Plugins/VirtualProduction/CameraCalibration/Source/CameraCalibrationEditor/Private/Calibrators/CameraLensDistortionAlgoPoints.cpp b/Engine/Plugins/VirtualProduction/CameraCalibration/Source/CameraCalibrationEditor/Private/Calibrators/CameraLensDistortionAlgoPoints.cpp
--- a/Engine/Plugins/VirtualProduction/CameraCalibration/Source/CameraCalibrationEditor/Private/Calibrators/CameraLensDistortionAlgoPoints.cpp
+++ b/Engine/Plugins/VirtualProduction/CameraCalibration/Source/CameraCalibrationEditor/Private/Calibrators/CameraLensDistortionAlgoPoints.cpp
@@ -31,7 +31,6 @@
 static TAutoConsoleVariable<bool> CVarFixExtrinsicsPoints(TEXT("LensDistortionPoints.FixExtrinsics"), false, TEXT("If true, the solver will fix the camera extrinsics to the user-provided camera poses"));
 static TAutoConsoleVariable<bool> CVarFixDistortionPoints(TEXT("LensDistortionPoints.FixDistortion"), false, TEXT("If true, the solver will not optimize distortion, and will use the input distortion values if any are given, or assume zero distortion otherwise."));
 static TAutoConsoleVariable<bool> CVarUseExtrinsicsGuessPoints(TEXT("LensDistortionPoints.UseExtrinsicsGuess"), false, TEXT("If true, the actual checkerboard and camera poses will be used when running the solver"));
-static TAutoConsoleVariable<bool> CVarGroupPointsByCameraPosePoints(TEXT("LensDistortionPoints.GroupPointsByCameraPose"), true, TEXT("If true, the points sent to the solver will be grouped together if they share the same camera pose."));
 #endif
 
 const int UCameraLensDistortionAlgoPoints::DATASET_VERSION = 1;
@@ -440,29 +439,6 @@ FDistortionCalibrationTask UCameraLensDistortionAlgoPoints::BeginCalibration(FTe
 		++NextPatternIndex;
 	}
 
-	// Validate that there are at least 4 patterns
-	if (Samples3d.Num() < 4)
-	{
-		OutErrorMessage = LOCTEXT("NotEnoughSamples", "At least 4 calibration patterns are required");
-		return CalibrationTask;
-	}
-
-	// Validate that each pattern has the same number of points
-	const int32 NumPointsInPattern = Samples3d[0].Points.Num();
-	for (int32 PatternIndex = 1; PatternIndex < Samples3d.Num(); ++PatternIndex)
-	{
-		if (Samples3d[PatternIndex].Points.Num() != NumPointsInPattern)
-		{
-			OutErrorMessage = LOCTEXT("DifferentNumPointsInPattern", "Every calibration pattern must have the same number of points");
-			return CalibrationTask;
-		}
-	}
-
-	if (CVarGroupPointsByCameraPosePoints.GetValueOnGameThread())
-	{
-		UE::CameraCalibration::Private::GroupPointsByCameraPose(Samples3d, Samples2d, CameraPoses);
-	}
-
 	// Because the calibration pattern is not coplanar, OpenCV requires an initial intrinsics guess
 	if (FMath::IsNearlyEqual(FocalLengthEstimate, 0.0) || FocalLengthEstimate < 0.0)
 	{
@@ -516,13 +492,17 @@ FDistortionCalibrationTask UCameraLensDistortionAlgoPoints::BeginCalibration(FTe
 	}
 
 	FDistortionInfo DistortionGuess;
-	LensFile->GetDistortionPoint(Focus, Zoom, DistortionGuess);
+	LensFile->EvaluateDistortionParameters(Focus, Zoom, DistortionGuess);
 
 	const TSubclassOf<ULensModel> Model = LensFile->LensInfo.LensModel;
 
 	Solver = NewObject<ULensDistortionSolver>(this, LensDistortionTool->GetSolverClass());
 
-	CalibrationTask = UE::Tasks::Launch(UE_SOURCE_LOCATION, [Solver = Solver, Model, Samples3d, Samples2d, ImageSize, FocalLength, ImageCenter, DistortionGuess, CameraPoses, PixelAspect, SolverFlags, Focus, Zoom]() mutable
+	// When capturing individual marker positions, there is no need to calibrate for a tracking offset of the prop
+	// Therefore there is no need to capture any target poses for the calibrator actor
+	TArray<FTransform> TargetPoses;
+
+	CalibrationTask = UE::Tasks::Launch(UE_SOURCE_LOCATION, [Solver = Solver, Model, Samples3d, Samples2d, ImageSize, FocalLength, ImageCenter, DistortionGuess, CameraPoses, TargetPoses, PixelAspect, SolverFlags, Focus, Zoom]() mutable
 	{
 		FDistortionCalibrationResult Result = Solver->Solve(
 			Samples3d,
@@ -532,6 +512,7 @@ FDistortionCalibrationTask UCameraLensDistortionAlgoPoints::BeginCalibration(FTe
 			ImageCenter,
 			DistortionGuess.Parameters,
 			CameraPoses,
+			TargetPoses,
 			Model,
 			PixelAspect,
 			SolverFlags
diff --git a/Engine/Plugins/VirtualProduction/CameraCalibration/Source/CameraCalibrationEditor/Private/TestCameraCalibration.cpp b/Engine/Plugins/VirtualProduction/CameraCalibration/Source/CameraCalibrationEditor/Private/TestCameraCalibration.cpp
--- a/Engine/Plugins/VirtualProduction/CameraCalibration/Source/CameraCalibrationEditor/Private/TestCameraCalibration.cpp
+++ b/Engine/Plugins/VirtualProduction/CameraCalibration/Source/CameraCalibrationEditor/Private/TestCameraCalibration.cpp
@@ -352,6 +352,8 @@ namespace UE::Private::CameraCalibration::AutomatedTests
 
 		ULensDistortionSolverOpenCV* TestSolver = NewObject<ULensDistortionSolverOpenCV>();
 
+		TArray<FTransform> TargetPoses;
+
 		FDistortionCalibrationResult Result = TestSolver->Solve(
 			NoisyObjectPoints,
 			NoisyImagePoints,
@@ -360,6 +362,7 @@ namespace UE::Private::CameraCalibration::AutomatedTests
 			CalibratedImageCenter,
 			DistortionGuess,
 			EstimatedCameraPoses,
+			TargetPoses,
 			USphericalLensModel::StaticClass(),
 			PixelAspect,
 			SolverFlags
@@ -847,6 +850,8 @@ namespace UE::Private::CameraCalibration::AutomatedTests
 
 		const FVector2D FocalLengthGuessInPixels = CalibrationTest.CameraProfile.ConvertFocalLengthToPixels(SolverSettings.FocalLengthGuess);
 
+		TArray<FTransform> TargetPoses;
+
 		FDistortionCalibrationResult Result = TestSolver->Solve(
 			CalibratorPoints3D,
 			ImagePoints2D,
@@ -855,6 +860,7 @@ namespace UE::Private::CameraCalibration::AutomatedTests
 			SolverSettings.ImageCenterGuess,
 			InitialDistortion,
 			CameraPoseGuesses,
+			TargetPoses,
 			USphericalLensModel::StaticClass(),
 			1.0,
 			SolverFlags
