CameraCalibration: Combine the functionality of the three native lens distortion "algo" classes. The three classes were functionally very similar, and really only differed in how they handled capturing the different calibration patterns. These different capture modes are now all supported by the lens distortion tool itself.
This CL also cleans up some of the UI elements in the lens distortion tab of the lens file editor, replacing most of the custom slate widgets with a details panel (with minor customization).

diff --git a/Engine/Plugins/VirtualProduction/CameraCalibration/Source/CameraCalibrationEditor/Private/AssetEditor/LensDistortionTool.cpp b/Engine/Plugins/VirtualProduction/CameraCalibration/Source/CameraCalibrationEditor/Private/AssetEditor/LensDistortionTool.cpp
--- a/Engine/Plugins/VirtualProduction/CameraCalibration/Source/CameraCalibrationEditor/Private/AssetEditor/LensDistortionTool.cpp
+++ b/Engine/Plugins/VirtualProduction/CameraCalibration/Source/CameraCalibrationEditor/Private/AssetEditor/LensDistortionTool.cpp
@@ -3,145 +3,99 @@
 #include "LensDistortionTool.h"
 
 #include "AssetToolsModule.h"
-#include "Calibrators/CameraCalibrationSolver.h"
+#include "CalibrationPointComponent.h"
+#include "Camera/CameraActor.h"
+#include "CameraCalibrationCheckerboard.h"
 #include "CameraCalibrationEditorLog.h"
 #include "CameraCalibrationSettings.h"
 #include "CameraCalibrationStepsController.h"
-#include "CameraCalibrationTypes.h"
-#include "CameraLensDistortionAlgo.h"
+#include "CameraCalibrationUtilsPrivate.h"
 #include "DesktopPlatformModule.h"
 #include "Dom/JsonObject.h"
 #include "Framework/Application/SlateApplication.h"
 #include "HAL/FileManager.h"
 #include "IDesktopPlatform.h"
 #include "ImageUtils.h"
+#include "JsonObjectConverter.h"
 #include "LensFile.h"
 #include "LensInfoStep.h"
 #include "Misc/DateTime.h"
 #include "Misc/MessageDialog.h"
-#include "Models/SphericalLensModel.h"
+#include "Modules/ModuleManager.h"
+#include "OpenCVHelper.h"
 #include "ScopedTransaction.h"
-#include "Serialization/JsonWriter.h"
 #include "Serialization/JsonSerializer.h"
-#include "UObject/Class.h"
-#include "UObject/Package.h"
-#include "UObject/UObjectIterator.h"
+#include "Serialization/JsonWriter.h"
+#include "SLensDistortionToolPanel.h"
 
 #define LOCTEXT_NAMESPACE "LensDistortionTool"
 
-#if WITH_EDITOR
-static TAutoConsoleVariable<bool> CVarSolveNodalOffset(TEXT("LensDistortionTool.SolveNodalOffset"), false, TEXT("If true, and if the calibration result contained a valid nodal offset, that offset will be added to the LensFile in addition to distortion."));
-#endif
-
 namespace UE::CameraCalibration::Private::LensDistortionTool
 {
 	static const FString SessionDateTimeField(TEXT("SessionDateTime"));
-	static const FString AlgoNameField(TEXT("AlgoName"));
+	static const FString Version(TEXT("Version"));
 }
 
 void ULensDistortionTool::Initialize(TWeakPtr<FCameraCalibrationStepsController> InCameraCalibrationStepController)
 {
-	CameraCalibrationStepsController = InCameraCalibrationStepController;
-
-	// Find available solver classes
+	WeakStepsController = InCameraCalibrationStepController;
 
+	// Discover all available solver classes and select the first one as the solver to use
 	TArray<UClass*> DerivedSolverClasses;
 	GetDerivedClasses(ULensDistortionSolver::StaticClass(), DerivedSolverClasses);
 
-	check(!DerivedSolverClasses.IsEmpty());
-
-	SetSolverClass(DerivedSolverClasses[0]);
-
-	// Find available algos
+	if (DerivedSolverClasses.Num() > 0)
+	{
+		SolverSettings.SolverClass = DerivedSolverClasses[0];
+	}
 
-	TArray<TSubclassOf<UCameraLensDistortionAlgo>> Algos;
+	// Find all actors in the current level that have calibration point components and select the first one as the starting calibrator 
+	TArray<AActor*> CalibratorActors;
+	UE::CameraCalibration::Private::FindActorsWithCalibrationComponents(CalibratorActors);
 
-	for (TObjectIterator<UClass> AlgoIt; AlgoIt; ++AlgoIt)
+	if (CalibratorActors.Num() > 0)
 	{
-		if (AlgoIt->IsChildOf(UCameraLensDistortionAlgo::StaticClass()) && !AlgoIt->HasAnyClassFlags(CLASS_Abstract | CLASS_Deprecated))
-		{
-			const UCameraLensDistortionAlgo* Algo = CastChecked<UCameraLensDistortionAlgo>(AlgoIt->GetDefaultObject());
-			AlgosMap.Add(Algo->FriendlyName(), TSubclassOf<UCameraLensDistortionAlgo>(*AlgoIt));
-
-			// If the algo uses an overlay material, create a new MID to use with that algo
-			if (UMaterialInterface* OverlayMaterial = Algo->GetOverlayMaterial())
-			{
-				AlgoOverlayMIDs.Add(Algo->FriendlyName(), UMaterialInstanceDynamic::Create(OverlayMaterial, GetTransientPackage()));
-			}
-		}
+		SetCalibrator(CalibratorActors[0]);
 	}
 
-	if (TSharedPtr<FCameraCalibrationStepsController> SharedStepsController = CameraCalibrationStepsController.Pin())
+	// Initialize the overlay material and texture
+	if (TSharedPtr<FCameraCalibrationStepsController> StepsController = WeakStepsController.Pin())
 	{
-		if (ULensFile* const LensFile = SharedStepsController->GetLensFile())
-		{
-			LensFile->OnLensFileModelChanged().AddUObject(this, &ULensDistortionTool::OnLensModelChanged);
+		UMaterialInterface* OverlayParent = TSoftObjectPtr<UMaterialInterface>(FSoftObjectPath(TEXT("/CameraCalibration/Materials/M_Coverage.M_Coverage"))).LoadSynchronous();
+		OverlayMID = UMaterialInstanceDynamic::Create(OverlayParent, GetTransientPackage());
 
-			UpdateAlgoMap(LensFile->LensInfo.LensModel);
-		}
-	}
+		const FIntPoint CompResolution = StepsController->GetCompRenderResolution();
 
-	BuildProgressWindowWidgets();
-}
+		OverlayTexture = UTexture2D::CreateTransient(CompResolution.X, CompResolution.Y, EPixelFormat::PF_B8G8R8A8);
+		UE::CameraCalibration::Private::ClearTexture(OverlayTexture);
 
-void ULensDistortionTool::OnLensModelChanged(const TSubclassOf<ULensModel>& LensModel)
-{
-	UpdateAlgoMap(LensModel);
-}
-
-void ULensDistortionTool::UpdateAlgoMap(const TSubclassOf<ULensModel>& LensModel)
-{
-	SupportedAlgosMap.Empty();
-	for (const TPair<FName, TSubclassOf<UCameraLensDistortionAlgo>>& AlgoPair : AlgosMap)
-	{
-		const UCameraLensDistortionAlgo* Algo = CastChecked<UCameraLensDistortionAlgo>(AlgoPair.Value->GetDefaultObject());
+		OverlayMID->SetTextureParameterValue(FName(TEXT("CoverageTexture")), OverlayTexture);
 
-		if (Algo->SupportsModel(LensModel))
-		{
-			SupportedAlgosMap.Add(AlgoPair);
-		}
-	}
-	if (DistortionWidget)
-	{
-		DistortionWidget->UpdateAlgosOptions();
+		StepsController->SetOverlayMaterial(OverlayMID);
 	}
 }
 
 void ULensDistortionTool::Shutdown()
 {
-	if (CurrentAlgo)
+	if (CalibrationTask.IsValid())
 	{
-		if (CalibrationTask.IsValid())
-		{
-			CurrentAlgo->CancelCalibration();
-			CalibrationTask = {};
-			ProgressWindow->HideWindow();
-		}
-
-		CurrentAlgo->Shutdown();
-		CurrentAlgo = nullptr;
-
-		EndCalibrationSession();
+		CancelCalibration();
+		CalibrationTask = {};
+		DistortionWidget->Shutdown();
 	}
 
-	if (TSharedPtr<FCameraCalibrationStepsController> SharedStepsController = CameraCalibrationStepsController.Pin())
-	{
-		if (ULensFile* const LensFile = SharedStepsController->GetLensFile())
-		{
-			LensFile->OnLensFileModelChanged().RemoveAll(this);
-		}
-	}
+	EndCalibrationSession();
 }
 
-void ULensDistortionTool::Tick(float DeltaTime)
+TSharedRef<SWidget> ULensDistortionTool::BuildUI()
 {
-	if (CurrentAlgo)
-	{
-		CurrentAlgo->Tick(DeltaTime);
-	}
+	DistortionWidget = SNew(SLensDistortionToolPanel, this, WeakStepsController);
+	return DistortionWidget.ToSharedRef();
+}
 
+void ULensDistortionTool::Tick(float DeltaTime)
+{
 	// A valid task handle implies that there is an asynchronous calibration happening on another thread.
-	// The tool will poll the task to determine when it has finished so that the results can be saved.
 	if (CalibrationTask.IsValid())
 	{
 		if (CalibrationTask.IsCompleted())
@@ -150,31 +104,32 @@ void ULensDistortionTool::Tick(float DeltaTime)
 			CalibrationResult = CalibrationTask.GetResult();
 			CalibrationTask = {};
 
-			if (!CalibrationResult.ErrorMessage.IsEmpty())
+			FText TaskCompletionText;
+			if (CalibrationResult.ErrorMessage.IsEmpty())
 			{
-				const FText Message = FText::Format(LOCTEXT("CalibrationErrorResult", "Calibration Error: {0}"), CalibrationResult.ErrorMessage);
-				ProgressTextWidget->SetText(Message);
+				FNumberFormattingOptions Options;
+				Options.MinimumFractionalDigits = 3;
+				Options.MaximumFractionalDigits = 3;
+
+				TaskCompletionText = FText::Format(LOCTEXT("CalibrationResultReprojectionError", "Reprojection Error: {0} pixels"), FText::AsNumber(CalibrationResult.ReprojectionError, &Options));
 			}
 			else
 			{
-				// Update progress window with final reprojection error
-				FFormatOrderedArguments Arguments;
-				Arguments.Add(FText::FromString(FString::Printf(TEXT("%.3f"), CalibrationResult.ReprojectionError)));
-
-				const FText Message = FText::Format(LOCTEXT("CalibrationTaskResult", "Reprojection Error: {0} pixels"), Arguments);
-				ProgressTextWidget->SetText(Message);
+				TaskCompletionText = FText::Format(LOCTEXT("CalibrationResultErrorMessage", "Calibration Error: {0}"), CalibrationResult.ErrorMessage);
 			}
 
-			OkayButton->SetEnabled(true);
+			DistortionWidget->UpdateProgressText(TaskCompletionText);
+			DistortionWidget->MarkProgressFinished();
 		}
 		else
 		{
+			// Update the calibration status in the progress window
 			FText StatusText = FText::GetEmpty();
-			const bool bIsStatusNew = CurrentAlgo->GetCalibrationStatus(StatusText);
+			const bool bIsStatusNew = GetCalibrationStatus(StatusText);
 
 			if (bIsStatusNew)
 			{
-				ProgressTextWidget->SetText(StatusText);
+				DistortionWidget->UpdateProgressText(StatusText);
 			}
 		}
 	}
@@ -182,303 +137,445 @@ void ULensDistortionTool::Tick(float DeltaTime)
 
 bool ULensDistortionTool::OnViewportClicked(const FGeometry& MyGeometry, const FPointerEvent& MouseEvent)
 {
-	if (!bIsActive)
+	// We only respond to left clicks
+	if (MouseEvent.GetEffectingButton() != EKeys::LeftMouseButton)
 	{
 		return false;
 	}
 
-	if (!CurrentAlgo)
+	// Block user interaction with the simulcam viewport while an async calibration task is executing
+	if (CalibrationTask.IsValid())
 	{
 		return false;
 	}
 
-	// Block user interaction with the simulcam viewport while an async calibration task is executing
-	if (CalibrationTask.IsValid())
+	TSharedPtr<FCameraCalibrationStepsController> StepsController = WeakStepsController.Pin();
+	if (!StepsController)
 	{
 		return false;
 	}
 
-	return CurrentAlgo->OnViewportClicked(MyGeometry, MouseEvent);
-}
+	// Create new row of calibration data
+	TSharedPtr<FCalibrationRow> NewRow = MakeShared<FCalibrationRow>();
 
-TSharedRef<SWidget> ULensDistortionTool::BuildUI()
-{
-	DistortionWidget = SNew(SLensDistortionToolPanel, this);
-	return DistortionWidget.ToSharedRef();
-}
+	// If capturing a single point, add the clicked point to the new row
+	if (CaptureSettings.CalibrationPattern == ECalibrationPattern::Points)
+	{
+		const bool bDetectionResult = DetectPoint(MyGeometry, MouseEvent, NewRow);
+		if (!bDetectionResult)
+		{
+			return true; // Though unsuccessful, the user input was handled
+		}
+	}
+	else
+	{
+		// If capturing a calibration pattern, read the media texture into CPU memory to send to one of the pattern detection algorithms
+		FText OutErrorMessage;
 
-void ULensDistortionTool::BuildProgressWindowWidgets()
-{
-	ProgressWindow = SNew(SWindow)
-		.Title(LOCTEXT("ProgressWindowTitle", "Distortion Calibration Progress"))
-		.SizingRule(ESizingRule::Autosized)
-		.IsTopmostWindow(true)
-		.HasCloseButton(false)
-		.SupportsMaximize(false)
-		.SupportsMinimize(true);
-
-	ProgressTextWidget = SNew(STextBlock).Text(FText::GetEmpty());
-
-	OkayButton = SNew(SButton)
-		.IsEnabled(false)
-		.HAlign(HAlign_Center)
-		.VAlign(VAlign_Center)
-		.Text(LOCTEXT("OkText", "Ok"))
-		.OnClicked_UObject(this, &ULensDistortionTool::OnOkPressed);
-
-	TSharedRef<SWidget> WindowContent = SNew(SVerticalBox)
-
-		// Text widget to display the current progress of the calibration
-		+ SVerticalBox::Slot()
-		.HAlign(EHorizontalAlignment::HAlign_Center)
-		.VAlign(EVerticalAlignment::VAlign_Center)
-		[
-			ProgressTextWidget.ToSharedRef()
-		]
-
-		// Ok and Cancel buttons
-		+ SVerticalBox::Slot()
-		.HAlign(EHorizontalAlignment::HAlign_Center)
-		.VAlign(EVerticalAlignment::VAlign_Center)
-		[
-			SNew(SHorizontalBox)
-
-			+ SHorizontalBox::Slot()
-			.AutoWidth()
-			[
-				OkayButton.ToSharedRef()
-			]
-
-			+ SHorizontalBox::Slot()
-			.AutoWidth()
-			[
-				SNew(SButton)
-				.Text(LOCTEXT("CancelText", "Cancel"))
-				.HAlign(HAlign_Center)
-				.VAlign(VAlign_Center)
-				.OnClicked_UObject(this, &ULensDistortionTool::OnCancelPressed)
-			]
-		];
-
-	ProgressWindow->SetContent(WindowContent);
-
-	// Create the window, but start with it hidden. When the user initiates a calibration, the progress window will be shown.
-	FSlateApplication::Get().AddWindow(ProgressWindow.ToSharedRef());
-	ProgressWindow->HideWindow();
-}
+		TArray<FColor> Pixels;
+		FIntPoint ImageSize;
+		if (!StepsController->ReadMediaPixels(Pixels, ImageSize, OutErrorMessage, ESimulcamViewportPortion::CameraFeed))
+		{
+			FMessageDialog::Open(EAppMsgType::Ok, OutErrorMessage, LOCTEXT("ErrorReadingMedia", "Media Error"));
+			return true; // Though unsuccessful, the user input was handled
+		}
 
-bool ULensDistortionTool::DependsOnStep(UCameraCalibrationStep* Step) const
-{
-	return Cast<ULensInfoStep>(Step) != nullptr;
-}
+		// Detect the selected pattern in the media image
+		bool bDetectionResult = false;
+		if (CaptureSettings.CalibrationPattern == ECalibrationPattern::Checkerboard)
+		{
+			bDetectionResult = DetectCheckerboardPattern(Pixels, ImageSize, NewRow);
+		}
+		else if (CaptureSettings.CalibrationPattern == ECalibrationPattern::Aruco)
+		{
+			bDetectionResult = DetectArucoPattern(Pixels, ImageSize, NewRow);
+		}
 
-void ULensDistortionTool::Activate()
-{
-	// Nothing to do if it is already active.
-	if (bIsActive)
+		if (!bDetectionResult)
+		{
+			return true; // Though unsuccessful, the user input was handled
+		}
+
+		// Save an image view of the captured frame
+		FImageView ImageView = FImageView(Pixels.GetData(), ImageSize.X, ImageSize.Y, ERawImageFormat::BGRA8);
+		ImageView.CopyTo(NewRow->MediaImage);
+	}
+
+	NewRow->Index = AdvanceSessionRowIndex();
+	NewRow->Pattern = CaptureSettings.CalibrationPattern;
+
+	NewRow->CameraPose = FTransform::Identity;
+	if (const ACameraActor* Camera = StepsController->GetCamera())
 	{
-		return;
+		if (const UCameraComponent* CameraComponent = Camera->GetCameraComponent())
+		{
+			NewRow->CameraPose = CameraComponent->GetComponentToWorld();
+		}
 	}
 
-	bIsActive = true;
-}
+	Dataset.CalibrationRows.Add(NewRow);
 
-void ULensDistortionTool::Deactivate()
-{
-	bIsActive = false;
-}
+	// Notify the ListView of the new data
+	if (DistortionWidget)
+	{
+		DistortionWidget->RefreshListView();
+	}
 
-UClass* ULensDistortionTool::GetSolverClass()
-{
-	return SolverClass;
-}
+	// Export the data for this row to a .json file on disk
+	ExportCalibrationRow(NewRow);
 
-void ULensDistortionTool::SetSolverClass(UClass* InSolverClass)
-{
-	SolverClass = InSolverClass;
+	ExportSessionData();
+
+	return true;
 }
 
-void ULensDistortionTool::ResetAlgo()
+bool ULensDistortionTool::DetectCheckerboardPattern(TArray<FColor>& Pixels, FIntPoint Size, TSharedPtr<FCalibrationRow> OutRow)
 {
-	// Remove old Algo
-	if (CurrentAlgo)
+	const FText ErrorTitle = LOCTEXT("CaptureError", "Capture Error");
+
+	// The selected calibrator must be a checkerboard actor
+	ACameraCalibrationCheckerboard* Checkerboard = Cast<ACameraCalibrationCheckerboard>(CaptureSettings.Calibrator.Get());
+	if (!Checkerboard)
 	{
-		CurrentAlgo->Shutdown();
-		CurrentAlgo = nullptr;
+		const FText ErrorMessage = LOCTEXT("CheckerboardActorRequiredError", "The selected calibrator must be actor must be a Camera Calibration Checkerboard actor.");
+		FMessageDialog::Open(EAppMsgCategory::Error, EAppMsgType::Ok, ErrorMessage, ErrorTitle);
+		return false;
+	}
+
+	OutRow->CheckerboardDimensions = FIntPoint(Checkerboard->NumCornerCols, Checkerboard->NumCornerRows);
+
+	TArray<FVector2f> DetectedCorners;
+	const bool bCornersFound = FOpenCVHelper::IdentifyCheckerboard(Pixels, Size, OutRow->CheckerboardDimensions, DetectedCorners);
+
+	if (!bCornersFound || DetectedCorners.IsEmpty())
+	{
+		const FText ErrorMessage = FText::Format(LOCTEXT("NoCheckerboardError", "Failed to detect a {0}x{1} checkerboard in the image."), Checkerboard->NumCornerCols, Checkerboard->NumCornerRows);
+		FMessageDialog::Open(EAppMsgCategory::Error, EAppMsgType::Ok, ErrorMessage, ErrorTitle);
+		return false;
+	}
 
-		EndCalibrationSession();
+	for (const FVector2f& Corner : DetectedCorners)
+	{
+		OutRow->ImagePoints.Points.Add(FVector2D(Corner));
 	}
 
-	// Set the tool overlay pass' material to the MID associate with the current algo
-	if (TSharedPtr<FCameraCalibrationStepsController> StepsController = CameraCalibrationStepsController.Pin())
+	// Fill out the checkerboard's 3D points
+	const FVector TopLeft = Checkerboard->TopLeft->GetComponentLocation();
+	const FVector TopRight = Checkerboard->TopRight->GetComponentLocation();
+	const FVector BottomLeft = Checkerboard->BottomLeft->GetComponentLocation();
+
+	const FVector RightVector = TopRight - TopLeft;
+	const FVector DownVector = BottomLeft - TopLeft;
+
+	const float HorizontalStep = (Checkerboard->NumCornerCols > 1) ? (1.0f / (Checkerboard->NumCornerCols - 1)) : 0.0f;
+	const float VerticalStep = (Checkerboard->NumCornerRows > 1) ? (1.0f / (Checkerboard->NumCornerRows - 1)) : 0.0f;
+
+	for (int32 RowIdx = 0; RowIdx < Checkerboard->NumCornerRows; ++RowIdx)
 	{
-		StepsController->SetOverlayEnabled(false);
-		StepsController->SetOverlayMaterial(nullptr);
+		for (int32 ColIdx = 0; ColIdx < Checkerboard->NumCornerCols; ++ColIdx)
+		{
+			const FVector PointLocation = TopLeft + (RightVector * ColIdx * HorizontalStep) + (DownVector * RowIdx * VerticalStep);
+			OutRow->ObjectPoints.Points.Add(PointLocation);
+		}
 	}
+
+	// Update the coverage overlay with the latest checkerboard corners
+	TArray<FVector2D> CameraFeedAdjustedCorners = OutRow->ImagePoints.Points;
+	const FIntPoint OverlayTextureSize = FIntPoint(OverlayTexture->GetSizeX(), OverlayTexture->GetSizeY());
+	RescalePoints(CameraFeedAdjustedCorners, OverlayTextureSize, Size);
+
+	FOpenCVHelper::DrawCheckerboardCorners(DetectedCorners, OutRow->CheckerboardDimensions, OverlayTexture);
+
+	if (TSharedPtr<FCameraCalibrationStepsController> StepsController = WeakStepsController.Pin())
+	{
+		StepsController->RefreshOverlay();
+	}
+
+	OutRow->TargetPose.SetLocation(TopLeft);
+
+	FRotator BoardRotation = FRotationMatrix::MakeFromYZ(TopRight - TopLeft, TopLeft - BottomLeft).Rotator();
+	OutRow->TargetPose.SetRotation(BoardRotation.Quaternion());
+
+	return true;
 }
 
-void ULensDistortionTool::SetAlgo(const FName& AlgoName)
+bool ULensDistortionTool::DetectArucoPattern(TArray<FColor>& Pixels, FIntPoint Size, TSharedPtr<FCalibrationRow> OutRow)
 {
-	// Find the algo class
+	const FText ErrorTitle = LOCTEXT("CaptureError", "Capture Error");
 
-	//@todo replace with Find to avoid double search
+	// Detect the aruco dictionary to use by looking at the names of the calibration points on the selected calibrator actor
+	AActor* CalibratorActor = CaptureSettings.Calibrator.Get();
+	if (!CalibratorActor)
+	{
+		const FText ErrorMessage = LOCTEXT("NoCalibratorError", "Please select a valid calibrator actor.");
+		FMessageDialog::Open(EAppMsgCategory::Error, EAppMsgType::Ok, ErrorMessage, ErrorTitle);
+		return false;
+	}
 
-	if (!AlgosMap.Contains(AlgoName))
+	EArucoDictionary Dictionary = UE::CameraCalibration::Private::GetArucoDictionaryForCalibrator(CalibratorActor);
+	if (Dictionary == EArucoDictionary::None)
 	{
-		return;
+		const FText ErrorMessage = LOCTEXT("NoArucoDictionaryError", "The calibration components of the selected calibrator do not specify a valid Aruco dictionary.");
+		FMessageDialog::Open(EAppMsgCategory::Error, EAppMsgType::Ok, ErrorMessage, ErrorTitle);
+		return false;
 	}
 
-	TSubclassOf<UCameraLensDistortionAlgo>& AlgoClass = AlgosMap[AlgoName];
-		
-	// If it is the same as the existing one, do nothing.
-	if (!CurrentAlgo && !AlgoClass)
+	// Identify any aruco markers matching the current dictionary in the media image
+	TArray<FArucoMarker> IdentifiedMarkers;
+	bool bResult = FOpenCVHelper::IdentifyArucoMarkers(Pixels, Size, Dictionary, IdentifiedMarkers);
+
+	if (!bResult || IdentifiedMarkers.IsEmpty())
 	{
-		return;
+		const FText ErrorMessage = LOCTEXT("NoArucoMarkersFoundError", "Failed to detect any aruco markers in the image belonging to the dictionary of the selected calibrator.");
+		FMessageDialog::Open(EAppMsgCategory::Error, EAppMsgType::Ok, ErrorMessage, ErrorTitle);
+		return false;
 	}
-	else if (CurrentAlgo && (CurrentAlgo->GetClass() == AlgoClass))
+
+	// For each identified marker, search the calibration components to find the subpoints matching that marker and the 3D location of each of its corners
+	TArray<FArucoCalibrationPoint> ArucoCalibrationPoints;
+	ArucoCalibrationPoints.Reserve(IdentifiedMarkers.Num());
+
+	UpdateCalibrationComponents();
+
+	for (const FArucoMarker& Marker : IdentifiedMarkers)
 	{
-		return;
+		FArucoCalibrationPoint ArucoCalibrationPoint;
+		if (UE::CameraCalibration::Private::FindArucoCalibrationPoint(CalibrationComponents, Dictionary, Marker, ArucoCalibrationPoint))
+		{
+			ArucoCalibrationPoints.Add(ArucoCalibrationPoint);
+		}
 	}
 
-	// Remove old Algo
-	if (CurrentAlgo)
+	if (ArucoCalibrationPoints.Num() > 0)
 	{
-		CurrentAlgo->Shutdown();
-		CurrentAlgo = nullptr;
+		for (const FArucoCalibrationPoint& Marker : ArucoCalibrationPoints)
+		{
+			for (int32 CornerIndex = 0; CornerIndex < 4; ++CornerIndex)
+			{
+				OutRow->ObjectPoints.Points.Add(Marker.Corners3D[CornerIndex]);
+				OutRow->ImagePoints.Points.Add(FVector2D(Marker.Corners2D[CornerIndex]));
+			}
+		}
+
+		FArucoCalibrationPoint& FirstAruco = ArucoCalibrationPoints[0];
+
+		FVector TopLeft = FirstAruco.Corners3D[0];
+		FVector TopRight = FirstAruco.Corners3D[1];
+		FVector BottomLeft = FirstAruco.Corners3D[3];
 
-		EndCalibrationSession();
+		OutRow->TargetPose.SetLocation(TopLeft);
+
+		FRotator FirstMarkerRotation = FRotationMatrix::MakeFromYZ(TopRight - TopLeft, TopLeft - BottomLeft).Rotator();
+		OutRow->TargetPose.SetRotation(FirstMarkerRotation.Quaternion());
 	}
 
-	// If AlgoClass is none, we're done here.
-	if (!AlgoClass)
+	FOpenCVHelper::DrawArucoMarkers(IdentifiedMarkers, OverlayTexture);
+
+	if (TSharedPtr<FCameraCalibrationStepsController> StepsController = WeakStepsController.Pin())
 	{
-		return;
+		StepsController->RefreshOverlay();
 	}
 
-	// Create new algo
-	CurrentAlgo = NewObject<UCameraLensDistortionAlgo>(
-		GetTransientPackage(),
-		AlgoClass,
-		MakeUniqueObjectName(GetTransientPackage(), AlgoClass));
+	return true;
+}
 
-	if (CurrentAlgo)
+bool ULensDistortionTool::DetectPoint(const FGeometry& MyGeometry, const FPointerEvent& MouseEvent, TSharedPtr<FCalibrationRow> OutRow)
+{
+	TSharedPtr<FCameraCalibrationStepsController> StepsController = WeakStepsController.Pin();
+	if (!StepsController)
 	{
-		CurrentAlgo->Initialize(this);
+		return false;
 	}
 
-	// Set the tool overlay pass' material to the MID associate with the current algo
-	if (CameraCalibrationStepsController.IsValid())
+	// When the user captures the first calibration point, the 3D locations for the remaining points are cached to ensure they all come from the same frame of data
+	// The media is also paused to ensure that the user is able to more easily capture the 2D location of each subsequent calibration point
+	UpdateCalibrationComponents();
+	
+	if (CalibrationComponentIndex == 0)
 	{
-		TSharedPtr<FCameraCalibrationStepsController> StepsController = CameraCalibrationStepsController.Pin();
-		StepsController->SetOverlayEnabled(false);
-		StepsController->SetOverlayMaterial(GetOverlayMID());
+		StepsController->Pause();
+
+		// Cache the camera pose and component locations of the remaining components
+		CachedComponentLocations.Reset();
+		for (TWeakObjectPtr<UCalibrationPointComponent> Component : CalibrationComponents)
+		{
+			CachedComponentLocations.Add(Component->GetComponentLocation());
+		}
 	}
-}
 
-UCameraLensDistortionAlgo* ULensDistortionTool::GetAlgo() const
-{
-	return CurrentAlgo;
+	// Calculate the location where the user clicked in the viewport
+	FVector2f NormalizedClickPosition;
+	if (!StepsController->CalculateNormalizedMouseClickPosition(MyGeometry, MouseEvent, NormalizedClickPosition, ESimulcamViewportPortion::CameraFeed))
+	{
+		return false;
+	}
+
+	const FIntPoint ImageSize = StepsController->GetCameraFeedSize();
+	const FVector2D ImagePoint = FVector2D(NormalizedClickPosition * ImageSize);
+
+	OutRow->ImagePoints.Points.Add(ImagePoint);
+
+	OutRow->ObjectPoints.Points.Add(CachedComponentLocations[CalibrationComponentIndex]);
+
+	// Advance the component index, and if it loops around, resume playing the media (which was paused after capturing the first point)
+	SetComponentIndex(CalibrationComponentIndex + 1);
+
+	if (CalibrationComponentIndex == 0)
+	{
+		StepsController->Play();
+	}
+
+	return true;
 }
 
-void ULensDistortionTool::OnSaveCurrentCalibrationData()
+void ULensDistortionTool::CalibrateLens()
 {
-	if (!CurrentAlgo)
+	TSharedPtr<FCameraCalibrationStepsController> StepsController = WeakStepsController.Pin();
+	if (!StepsController)
 	{
 		return;
 	}
 
-	const FText TitleError = LOCTEXT("LensCalibrationError", "Lens Calibration Error");
-	const FText UnknownError = LOCTEXT("UnknownError", "An unknown error occurred initiating the distortion calibration. Check the output log for details.");
+	ULensFile* LensFile = StepsController->GetLensFile();
+	if (!LensFile)
+	{
+		return;
+	}
+
+	const FText TitleError = LOCTEXT("CalibrationErrorTitle", "Calibration Error");
 
-	if (CurrentAlgo->SupportsAsyncCalibration())
+	if (Dataset.CalibrationRows.Num() < 1)
 	{
-		FText ErrorMessage;
-		CalibrationTask = CurrentAlgo->BeginCalibration(ErrorMessage);
+		FMessageDialog::Open(EAppMsgType::Ok, LOCTEXT("EmptyDatasetError", "The calibration dataset is empty. Please select a valid dataset or capture a new one before calibrating."), TitleError);
+		return;
+	}
 
-		if (!CalibrationTask.IsValid())
-		{
-			if (!ErrorMessage.IsEmpty())
-			{
-				FMessageDialog::Open(EAppMsgType::Ok, ErrorMessage, TitleError);
-			}
-			else
-			{
-				FMessageDialog::Open(EAppMsgType::Ok, UnknownError, TitleError);
-			}
-			return;
-		}
+	const float PhysicalSensorWidth = StepsController->GetLensFileEvaluationInputs().Filmback.SensorWidth;
+	const float PixelAspect = LensFile->LensInfo.SqueezeFactor;
 
-		// TODO: This text is temporary. Update to register a delegate with the algo to provide implementation-specific status text
-		ProgressTextWidget->SetText(LOCTEXT("CalibrationProgressText", "Calibrating Lens Distortion..."));
+	const float DesqueezedSensorWidth = PhysicalSensorWidth * PixelAspect;
 
-		// Ensure that the Ok button is disabled and show the progress window
-		OkayButton->SetEnabled(false);
-		ProgressWindow->ShowWindow();
+	if (FMath::IsNearlyZero(DesqueezedSensorWidth))
+	{
+		FMessageDialog::Open(EAppMsgType::Ok, LOCTEXT("InvalidSensorWidthError", "The sensor width and squeeze factor in the camera settings must both be greater than zero. Please enter a valid value."), TitleError);
+		return;
+	}
 
-		DistortionWidget->SetEnabled(false);
+	if (!SolverSettings.FocalLengthGuess.IsSet() || FMath::IsNearlyZero(SolverSettings.FocalLengthGuess.GetValue()))
+	{
+		FMessageDialog::Open(EAppMsgType::Ok, LOCTEXT("InvalidFocalLengthError", "Please enter a valid estimate for the focal length of the lens (in mm)."), TitleError);
+		return;
 	}
-	else
+
+	const FIntPoint ImageSize = LensFile->CameraFeedInfo.GetDimensions();
+
+	const double FocalLengthEstimateValue = SolverSettings.FocalLengthGuess.GetValue();
+	const double Fx = (FocalLengthEstimateValue / DesqueezedSensorWidth) * ImageSize.X;
+
+	// When operating on a desqueezed image, we expect our pixel aspect to be square, so horizontal and vertical field of view are assumed to be equal (i.e. Fx == Fy)
+	FVector2D FocalLength = FVector2D(Fx);
+	FVector2D ImageCenter = FVector2D((ImageSize.X - 1) * 0.5, (ImageSize.Y - 1) * 0.5);
+
+	TArray<FObjectPoints> Samples3d;
+	Samples3d.Reserve(Dataset.CalibrationRows.Num());
+
+	TArray<FImagePoints> Samples2d;
+	Samples2d.Reserve(Dataset.CalibrationRows.Num());
+
+	TArray<FTransform> CameraPoses;
+	CameraPoses.Reserve(Dataset.CalibrationRows.Num());
+
+	TArray<FTransform> TargetPoses;
+	TargetPoses.Reserve(Dataset.CalibrationRows.Num());
+
+	// Extract the 3D points, 2D points, and camera poses from each row to pass to the solver
+	for (const TSharedPtr<FCalibrationRow>& Row : Dataset.CalibrationRows)
 	{
-		TSubclassOf<ULensModel> LensModel;
-		bool bResult = CurrentAlgo->GetLensDistortion(
-			CalibrationResult.EvaluatedFocus, 
-			CalibrationResult.EvaluatedZoom, 
-			CalibrationResult.Parameters, 
-			CalibrationResult.FocalLength, 
-			CalibrationResult.ImageCenter, 
-			LensModel, 
-			CalibrationResult.ReprojectionError, 
-			CalibrationResult.ErrorMessage);
-
-		if (!bResult)
-		{
-			if (CalibrationResult.ErrorMessage.IsEmpty())
-			{
-				CalibrationResult.ErrorMessage = UnknownError;
-			}
+		Samples3d.Add(Row->ObjectPoints);
+		Samples2d.Add(Row->ImagePoints);
+		CameraPoses.Add(Row->CameraPose);
+		TargetPoses.Add(Row->TargetPose);
+	}
 
-			FMessageDialog::Open(EAppMsgType::Ok, CalibrationResult.ErrorMessage, TitleError);
-			return;
-		}
+	ECalibrationFlags SolverFlags = ECalibrationFlags::None;
+	EnumAddFlags(SolverFlags, ECalibrationFlags::UseIntrinsicGuess);
 
-		// Update progress window with final reprojection error
-		ProgressWindow->ShowWindow();
+	if (CaptureSettings.bIsCameraTracked)
+	{
+		EnumAddFlags(SolverFlags, ECalibrationFlags::UseExtrinsicGuess);
+	}
 
-		FFormatOrderedArguments Arguments;
-		Arguments.Add(FText::FromString(FString::Printf(TEXT("%.3f"), CalibrationResult.ReprojectionError)));
+	if (CaptureSettings.bIsCalibratorTracked && CaptureSettings.CalibrationPattern != ECalibrationPattern::Points)
+	{
+		EnumAddFlags(SolverFlags, ECalibrationFlags::SolveTargetOffset);
+	}
 
-		const FText Message = FText::Format(LOCTEXT("ReprojectionError", "Reprojection Error: {0} pixels"), Arguments);
-		ProgressTextWidget->SetText(Message);
+	if (SolverSettings.bFixFocalLength)
+	{
+		EnumAddFlags(SolverFlags, ECalibrationFlags::FixFocalLength);
+	}
 
-		OkayButton->SetEnabled(true);
+	if (SolverSettings.bFixImageCenter)
+	{
+		EnumAddFlags(SolverFlags, ECalibrationFlags::FixPrincipalPoint);
 	}
-}
 
-FReply ULensDistortionTool::OnCancelPressed()
-{
-	CurrentAlgo->CancelCalibration();
+	if (SolverSettings.bFixDistortion)
+	{
+		EnumAddFlags(SolverFlags, ECalibrationFlags::FixDistortion);
+	}
 
-	CalibrationTask = {};
-	ProgressWindow->HideWindow();
-	DistortionWidget->SetEnabled(true);
+	const FLensFileEvaluationInputs LensFileEvalInputs = StepsController->GetLensFileEvaluationInputs();
 
-	return FReply::Handled();
-}
+	FDistortionInfo DistortionGuess;
+	LensFile->EvaluateDistortionParameters(LensFileEvalInputs.Focus, LensFileEvalInputs.Zoom, DistortionGuess);
 
-FReply ULensDistortionTool::OnOkPressed()
-{
-	SaveCalibrationResult();
+	const TSubclassOf<ULensModel> Model = LensFile->LensInfo.LensModel;
 
-	ProgressWindow->HideWindow();
-	DistortionWidget->SetEnabled(true);
+	Solver = NewObject<ULensDistortionSolver>(GetTransientPackage(), SolverSettings.SolverClass);
 
-	return FReply::Handled();
+	CalibrationTask = UE::Tasks::Launch(UE_SOURCE_LOCATION, [InSolver = this->Solver, Model, Samples3d, Samples2d, ImageSize, FocalLength, ImageCenter, DistortionGuess, CameraPoses, TargetPoses, PixelAspect, SolverFlags, LensFileEvalInputs]() mutable
+		{
+			FDistortionCalibrationResult Result = InSolver->Solve(
+				Samples3d,
+				Samples2d,
+				ImageSize,
+				FocalLength,
+				ImageCenter,
+				DistortionGuess.Parameters,
+				CameraPoses,
+				TargetPoses,
+				Model,
+				PixelAspect,
+				SolverFlags
+			);
+
+			// CalibrateCamera() returns focal length and image center in pixels, but the result is expected to be normalized by the image size
+			Result.FocalLength.FxFy = Result.FocalLength.FxFy / ImageSize;
+			Result.ImageCenter.PrincipalPoint = Result.ImageCenter.PrincipalPoint / ImageSize;
+
+			// FZ inputs to LUT
+			Result.EvaluatedFocus = LensFileEvalInputs.Focus;
+			Result.EvaluatedZoom = LensFileEvalInputs.Zoom;
+
+			return Result;
+		});
+
+	DistortionWidget->OpenProgressWindow();
+
+	// All of the UI options should be disabled while the calibration task is running
+	DistortionWidget->SetEnabled(false);
 }
 
 void ULensDistortionTool::SaveCalibrationResult()
 {
-	ULensFile* LensFile = CameraCalibrationStepsController.Pin()->GetLensFile();
+	TSharedPtr<FCameraCalibrationStepsController> StepsController = WeakStepsController.Pin();
+	if (!StepsController)
+	{
+		return;
+	}
 
+	ULensFile* LensFile = StepsController->GetLensFile();
 	if (!LensFile)
 	{
 		return;
@@ -496,7 +593,7 @@ void ULensDistortionTool::SaveCalibrationResult()
 		CalibrationResult.STMap.DistortionMap = (ImportedImages.Num() > 0) ? Cast<UTexture>(ImportedImages[0]) : nullptr;
 	}
 
-	// Depending on the algo, it is possible that the result feature calibrated distortion parameters or an ST Map.
+	// The result may have featured calibrated distortion parameters or an ST Map.
 	// If the result contains any distortion parameters, then the results will be written as a distortion point in the Lens File
 	// Otherwise, if the result contains a valid ST Map, then it will be added to the Lens File
 	if (CalibrationResult.Parameters.Parameters.Num() > 0)
@@ -530,7 +627,7 @@ void ULensDistortionTool::SaveCalibrationResult()
 	}
 
 	// If the calibration result contains a meaningful nodal offset result, add it to the Lens File
-	if (CVarSolveNodalOffset.GetValueOnGameThread())
+	if (SolverSettings.CalibrationType == ECalibrationType::NodalOffset || SolverSettings.CalibrationType == ECalibrationType::Everything)
 	{
 		if (!CalibrationResult.NodalOffset.LocationOffset.Equals(FVector::ZeroVector) || !CalibrationResult.NodalOffset.RotationOffset.Equals(FQuat::Identity))
 		{
@@ -541,52 +638,186 @@ void ULensDistortionTool::SaveCalibrationResult()
 		}
 	}
 
-	if (UCameraLensDistortionAlgo* Algo = GetAlgo())
+	ClearCalibrationRows();
+}
+
+void ULensDistortionTool::CancelCalibration()
+{
+	if (Solver)
+	{
+		Solver->Cancel();
+	}
+}
+
+bool ULensDistortionTool::GetCalibrationStatus(FText& StatusText) const
+{
+	if (Solver)
+	{
+		return Solver->GetStatusText(StatusText);
+	}
+	return false;
+}
+
+void ULensDistortionTool::SetCalibrator(AActor* InCalibrator)
+{
+	CaptureSettings.Calibrator = InCalibrator;
+
+	if (!InCalibrator)
 	{
-		Algo->OnDistortionSavedToLens();
+		return;
 	}
+
+	CalibrationComponents.Reset();
+	UpdateCalibrationComponents();
+
+	// Initialize the component index used by the single point detection mode
+	SetComponentIndex(0);
+
+	CachedComponentLocations.Reserve(CalibrationComponents.Num());
 }
 
-FCameraCalibrationStepsController* ULensDistortionTool::GetCameraCalibrationStepsController() const
+void ULensDistortionTool::UpdateCalibrationComponents()
 {
-	if (!CameraCalibrationStepsController.IsValid())
+	bool bNeedsUpdate = false;
+
+	for (TWeakObjectPtr<UCalibrationPointComponent> Component : CalibrationComponents)
 	{
-		return nullptr;
+		if (!Component.IsValid())
+		{
+			bNeedsUpdate = true;
+			break;
+		}
 	}
 
-	return CameraCalibrationStepsController.Pin().Get();
+	if (CalibrationComponents.IsEmpty())
+	{
+		bNeedsUpdate = true;
+	}
+
+	if (bNeedsUpdate && CaptureSettings.Calibrator.IsValid())
+	{
+		// Find all of the calibration components attached to calibrator actor
+		TArray<UCalibrationPointComponent*> CalibrationPoints;
+		CaptureSettings.Calibrator->GetComponents(CalibrationPoints);
+
+		// Store weak references to all of the calibration components that have an attached scene component
+		CalibrationComponents.Reset();
+		for (UCalibrationPointComponent* CalibrationPoint : CalibrationPoints)
+		{
+			if (CalibrationPoint && CalibrationPoint->GetAttachParent())
+			{
+				CalibrationComponents.Add(CalibrationPoint);
+			}
+		}
+	}
 }
 
-TArray<FName> ULensDistortionTool::GetAlgos() const
+void ULensDistortionTool::ClearCalibrationRows()
 {
-	TArray<FName> OutKeys;
-	SupportedAlgosMap.GetKeys(OutKeys);
-	return OutKeys;
+	Dataset.CalibrationRows.Empty();
+
+	if (DistortionWidget)
+	{
+		DistortionWidget->RefreshListView();
+	}
+
+	// Reset the calibration component index to restart the pattern
+	SetComponentIndex(0);
+
+	RefreshCoverage();
+
+	// End the current calibration session (a new one will begin the next time a new row is added)
+	EndCalibrationSession();
 }
 
-bool ULensDistortionTool::IsActive() const
+void ULensDistortionTool::SetComponentIndex(int32 Index)
 {
-	return bIsActive;
+	CalibrationComponentIndex = Index;
+
+	// If the Index would loop around, reset it to 0
+	if (CalibrationComponentIndex == CalibrationComponents.Num())
+	{
+		CalibrationComponentIndex = 0;
+	}
+
+	UpdateCalibrationComponents();
+
+	if (CalibrationComponents.IsValidIndex(CalibrationComponentIndex))
+	{
+		CaptureSettings.NextPoint = FText::FromString(CalibrationComponents[CalibrationComponentIndex]->GetName());
+	}
 }
 
-UMaterialInstanceDynamic* ULensDistortionTool::GetOverlayMID() const
+void ULensDistortionTool::RefreshCoverage()
 {
-	if (CurrentAlgo)
+	UE::CameraCalibration::Private::ClearTexture(OverlayTexture);
+
+	for (const TSharedPtr<FCalibrationRow>& Row : Dataset.CalibrationRows)
+	{
+		if (Row->Pattern == ECalibrationPattern::Checkerboard)
+		{
+			TArray<FVector2D> CameraFeedAdjustedCorners = Row->ImagePoints.Points;
+			const FIntPoint OverlayTextureSize = FIntPoint(OverlayTexture->GetSizeX(), OverlayTexture->GetSizeY());
+			const FIntPoint ImageSize = FIntPoint(Row->MediaImage.SizeX, Row->MediaImage.SizeY);
+			RescalePoints(CameraFeedAdjustedCorners, OverlayTextureSize, ImageSize);
+
+			FOpenCVHelper::DrawCheckerboardCorners(CameraFeedAdjustedCorners, Row->CheckerboardDimensions, OverlayTexture);
+		}
+	}
+
+	// The coverage texture may have changed as a result of a change in size or pixel format.
+	// Therefore, the material parameter should be updated to ensure it is up to date.
+	if (OverlayTexture)
 	{
-		return AlgoOverlayMIDs.FindRef(CurrentAlgo->FriendlyName()).Get();
+		OverlayMID->SetTextureParameterValue(FName(TEXT("CoverageTexture")), OverlayTexture);
 	}
 
-	return nullptr;
+	if (TSharedPtr<FCameraCalibrationStepsController> StepsController = WeakStepsController.Pin())
+	{
+		StepsController->RefreshOverlay();
+	}
 }
 
-bool ULensDistortionTool::IsOverlayEnabled() const
+void ULensDistortionTool::RescalePoints(TArray<FVector2D>& Points, FIntPoint DebugTextureSize, FIntPoint CameraFeedSize)
 {
-	if (CurrentAlgo)
+	// It is possible that the size of the debug texture is different than the size of the camera feed.
+	// Therefore, the input points should be shifted so that they appear at the correct location in the debug image.
+	const FVector2D TopLeftCorner = (DebugTextureSize - CameraFeedSize) / 2.0f;
+
+	for (FVector2D& Point : Points)
 	{
-		return CurrentAlgo->IsOverlayEnabled();
+		Point += TopLeftCorner;
 	}
+}
 
-	return false;
+bool ULensDistortionTool::DependsOnStep(UCameraCalibrationStep* Step) const
+{
+	return Cast<ULensInfoStep>(Step) != nullptr;
+}
+
+void ULensDistortionTool::Activate()
+{
+	bIsActive = true;
+}
+
+void ULensDistortionTool::Deactivate()
+{
+	bIsActive = false;
+}
+
+bool ULensDistortionTool::IsActive() const
+{
+	return bIsActive;
+}
+
+UMaterialInstanceDynamic* ULensDistortionTool::GetOverlayMID() const
+{
+	return OverlayMID;
+}
+
+bool ULensDistortionTool::IsOverlayEnabled() const
+{
+	return CaptureSettings.bShowOverlay;
 }
 
 void ULensDistortionTool::StartCalibrationSession()
@@ -619,8 +850,7 @@ FString ULensDistortionTool::GetSessionSaveDir() const
 
 	const FString SessionDateString = SessionInfo.StartTime.ToString(TEXT("%Y-%m-%d"));
 	const FString SessionTimeString = SessionInfo.StartTime.ToString(TEXT("%H-%M-%S"));
-	const FString DatasetPrefix = TEXT("Dataset-") + CurrentAlgo->ShortName().ToString() + TEXT("Algorithm-");
-	const FString DatasetDir = DatasetPrefix + SessionTimeString;
+	const FString DatasetDir = SessionTimeString;
 
 	const FString ProjectSaveDir = FPaths::ProjectSavedDir() / TEXT("CameraCalibration") / TEXT("LensDistortion");
 
@@ -659,84 +889,6 @@ void ULensDistortionTool::DeleteExportedRow(const int32& RowIndex) const
 	}
 }
 
-void ULensDistortionTool::ExportCalibrationRow(int32 RowIndex, const TSharedRef<FJsonObject>& RowObject, const FImageView& RowImage)
-{
-	if (!GetDefault<UCameraCalibrationSettings>()->IsCalibrationDatasetImportExportEnabled())
-	{
-		return;
-	}
-
-	// Start a calibration session (if one is not currently active)
-	StartCalibrationSession();
-
-	// Assemble the path and filename for this row based on the session and row index
-	const FString PathName = GetSessionSaveDir();
-	const FString FileName = GetRowFilename(RowIndex) + FDateTime::Now().ToString(TEXT("%H-%M-%S"));
-
-	const FString JsonFileName = PathName / FileName + TEXT(".json");
-	const FString ImageFileName = PathName / FileName + TEXT(".png");
-
-	// Create and open a new Json file for writing, and initialize a JsonWriter to serialize the contents
-	if (TUniquePtr<FArchive> FileWriter = TUniquePtr<FArchive>(IFileManager::Get().CreateFileWriter(*JsonFileName)))
-	{
-		TSharedRef< TJsonWriter<> > JsonWriter = TJsonWriterFactory<>::Create(FileWriter.Get());
-
-		// Write the Json row data out and save the file
-		FJsonSerializer::Serialize(RowObject, JsonWriter);
-		FileWriter->Close();
-
-		UE_LOG(LogCameraCalibrationEditor, Verbose, TEXT("Lens Distortion Tool wrote to dataset row file: %s"), *JsonFileName);
-	}
-
-	// If the row has an image to export, save it out to a file
-	if (RowImage.RawData != nullptr)
-	{
-		FImageUtils::SaveImageByExtension(*ImageFileName, RowImage);
-	}
-}
-
-void ULensDistortionTool::ExportSessionData(const TSharedRef<FJsonObject>& SessionDataObject)
-{
-	using namespace UE::CameraCalibration::Private;
-
-	if (!GetDefault<UCameraCalibrationSettings>()->IsCalibrationDatasetImportExportEnabled())
-	{
-		return;
-	}
-
-	// Start a calibration session (if one is not currently active)
-	StartCalibrationSession();
-
-	// Assemble the path and filename for this row based on the session and row index
-	const FString PathName = GetSessionSaveDir();
-	const FString FileName = TEXT("SessionData");
-
-	const FString SessionFileName = PathName / FileName + TEXT(".ucamcalib");
-
-	// Delete the existing session data file (if it exists)
-	if (IFileManager::Get().FileExists(*SessionFileName))
-	{
-		IFileManager::Get().Delete(*SessionFileName);
-	}
-
-	// Create and open a new Json file for writing, and initialize a JsonWriter to serialize the contents
-	if (TUniquePtr<FArchive> FileWriter = TUniquePtr<FArchive>(IFileManager::Get().CreateFileWriter(*SessionFileName)))
-	{
-		TSharedRef< TJsonWriter<> > JsonWriter = TJsonWriterFactory<>::Create(FileWriter.Get());
-
-		const FString SessionDateTimeString = SessionInfo.StartTime.ToString(TEXT("%Y-%m-%d")) + TEXT("-") + SessionInfo.StartTime.ToString(TEXT("%H-%M-%S"));
-
-		SessionDataObject->SetStringField(LensDistortionTool::SessionDateTimeField, SessionDateTimeString);
-		SessionDataObject->SetStringField(LensDistortionTool::AlgoNameField, CurrentAlgo->FriendlyName().ToString());
-
-		// Write the Json row data out and save the file
-		FJsonSerializer::Serialize(SessionDataObject, JsonWriter);
-		FileWriter->Close();
-
-		UE_LOG(LogCameraCalibrationEditor, Verbose, TEXT("Lens Distortion Tool wrote to dataset session file: %s"), *SessionFileName);
-	}
-}
-
 void ULensDistortionTool::ImportCalibrationDataset()
 {
 	using namespace UE::CameraCalibration::Private;
@@ -747,7 +899,7 @@ void ULensDistortionTool::ImportCalibrationDataset()
 	}
 
 	// If there is existing calibration data that will be overwritten during import, ask the user to confirm that they want to continue
-	if (CurrentAlgo->HasCalibrationData())
+	if (Dataset.CalibrationRows.Num() > 0)
 	{
 		const FText ConfirmationMessage = LOCTEXT(
 			"ImportDatasetConfirmationMessage",
@@ -819,30 +971,6 @@ void ULensDistortionTool::ImportCalibrationDataset()
 				{
 					UE_LOG(LogCameraCalibrationEditor, Verbose, TEXT("Lens Distortion Tool failed to deserialize the date/time from the session file: %s"), *SessionFile);
 				}
-
-				// Import the algo name so we can set the appropriate algo before importing the row data
-				FString AlgoString;
-				if (JsonSessionData->TryGetStringField(LensDistortionTool::AlgoNameField, AlgoString))
-				{
-					// Ensure that the algo name matches one of the algos for this tool
-					const FName AlgoName = FName(*AlgoString);
-					if (AlgosMap.Contains(AlgoName))
-					{
-						SetAlgo(AlgoName);
-					}
-					else
-					{
-						const FText ErrorMessage = LOCTEXT("UnknownCalibrationAlgo", "The selected dataset does not represent a lens distortion calibration. Choose a different .ucamcalib dataset file from a lens distortion calibration.");
-						FMessageDialog::Open(EAppMsgType::Ok, ErrorMessage);
-						return;
-					}
-				}
-				else
-				{
-					UE_LOG(LogCameraCalibrationEditor, Verbose, TEXT("Lens Distortion Tool failed to deserialize the algo name from the session file: %s"), *SessionFile);
-				}
-
-				CurrentAlgo->ImportSessionData(JsonSessionData.ToSharedRef());
 			}
 			else
 			{
@@ -851,7 +979,7 @@ void ULensDistortionTool::ImportCalibrationDataset()
 		}
 	}
 
-	CurrentAlgo->PreImportCalibrationData();
+	Dataset.CalibrationRows.Empty();
 
 	// Initialize a maximum row index which will be used to set the current session row index if the user wants to add additional rows after importing
 	int32 MaxRowIndex = -1;
@@ -877,7 +1005,7 @@ void ULensDistortionTool::ImportCalibrationDataset()
 			TSharedPtr<FJsonObject> JsonRowData = MakeShared<FJsonObject>();
 			if (FJsonSerializer::Deserialize(JsonReader, JsonRowData))
 			{
-				int32 RowNum = CurrentAlgo->ImportCalibrationRow(JsonRowData.ToSharedRef(), RowImage);
+				int32 RowNum = ImportCalibrationRow(JsonRowData.ToSharedRef(), RowImage);
 				MaxRowIndex = FMath::Max(MaxRowIndex, RowNum);
 			}
 			else
@@ -887,7 +1015,14 @@ void ULensDistortionTool::ImportCalibrationDataset()
 		}
 	}
 
-	CurrentAlgo->PostImportCalibrationData();
+	// Sort imported calibration rows by row index
+	Dataset.CalibrationRows.Sort([](const TSharedPtr<FCalibrationRow>& LHS, const TSharedPtr<FCalibrationRow>& RHS) { return LHS->Index < RHS->Index; });
+
+	// Notify the ListView of the new data
+	if (DistortionWidget)
+	{
+		DistortionWidget->RefreshListView();
+	}
 
 	// Set the current session's start date/time and row index to match what was just imported to support adding/deleting rows
 	SessionInfo.bIsActive = true;
@@ -895,4 +1030,115 @@ void ULensDistortionTool::ImportCalibrationDataset()
 	SessionInfo.RowIndex = MaxRowIndex;
 }
 
+void ULensDistortionTool::ExportSessionData()
+{
+	if (!GetDefault<UCameraCalibrationSettings>()->IsCalibrationDatasetImportExportEnabled())
+	{
+		return;
+	}
+
+	using namespace UE::CameraCalibration::Private;
+
+	TSharedPtr<FJsonObject> JsonSessionData = MakeShared<FJsonObject>();
+	JsonSessionData->SetNumberField(LensDistortionTool::Version, 2);
+
+	// Start a calibration session (if one is not currently active)
+	StartCalibrationSession();
+
+	// Assemble the path and filename for this row based on the session and row index
+	const FString PathName = GetSessionSaveDir();
+	const FString FileName = TEXT("SessionData");
+
+	const FString SessionFileName = PathName / FileName + TEXT(".ucamcalib");
+
+	// Delete the existing session data file (if it exists)
+	if (IFileManager::Get().FileExists(*SessionFileName))
+	{
+		IFileManager::Get().Delete(*SessionFileName);
+	}
+
+	// Create and open a new Json file for writing, and initialize a JsonWriter to serialize the contents
+	if (TUniquePtr<FArchive> FileWriter = TUniquePtr<FArchive>(IFileManager::Get().CreateFileWriter(*SessionFileName)))
+	{
+		TSharedRef< TJsonWriter<> > JsonWriter = TJsonWriterFactory<>::Create(FileWriter.Get());
+
+		const FString SessionDateTimeString = SessionInfo.StartTime.ToString(TEXT("%Y-%m-%d")) + TEXT("-") + SessionInfo.StartTime.ToString(TEXT("%H-%M-%S"));
+
+		JsonSessionData->SetStringField(LensDistortionTool::SessionDateTimeField, SessionDateTimeString);
+
+		// Write the Json row data out and save the file
+		FJsonSerializer::Serialize(JsonSessionData.ToSharedRef(), JsonWriter);
+		FileWriter->Close();
+
+		UE_LOG(LogCameraCalibrationEditor, Verbose, TEXT("Lens Distortion Tool wrote to dataset session file: %s"), *SessionFileName);
+	}
+}
+
+void ULensDistortionTool::ExportCalibrationRow(TSharedPtr<FCalibrationRow> Row)
+{
+	if (!GetDefault<UCameraCalibrationSettings>()->IsCalibrationDatasetImportExportEnabled())
+	{
+		return;
+	}
+
+	if (const TSharedPtr<FJsonObject>& RowObject = FJsonObjectConverter::UStructToJsonObject<FCalibrationRow>(Row.ToSharedRef().Get()))
+	{
+		// Start a calibration session (if one is not currently active)
+		StartCalibrationSession();
+
+		// Assemble the path and filename for this row based on the session and row index
+		const FString PathName = GetSessionSaveDir();
+		const FString FileName = GetRowFilename(Row->Index) + FDateTime::Now().ToString(TEXT("%H-%M-%S"));
+
+		const FString JsonFileName = PathName / FileName + TEXT(".json");
+		const FString ImageFileName = PathName / FileName + TEXT(".png");
+
+		// Create and open a new Json file for writing, and initialize a JsonWriter to serialize the contents
+		if (TUniquePtr<FArchive> FileWriter = TUniquePtr<FArchive>(IFileManager::Get().CreateFileWriter(*JsonFileName)))
+		{
+			TSharedRef< TJsonWriter<> > JsonWriter = TJsonWriterFactory<>::Create(FileWriter.Get());
+
+			// Write the Json row data out and save the file
+			FJsonSerializer::Serialize(RowObject.ToSharedRef(), JsonWriter);
+			FileWriter->Close();
+
+			UE_LOG(LogCameraCalibrationEditor, Verbose, TEXT("Lens Distortion Tool wrote to dataset row file: %s"), *JsonFileName);
+		}
+
+		// If the row has an image to export, save it out to a file
+		FImageView ImageView = Row->MediaImage;
+		if (ImageView.RawData != nullptr)
+		{
+			FImageUtils::SaveImageByExtension(*ImageFileName, ImageView);
+		}
+	}
+}
+
+int32 ULensDistortionTool::ImportCalibrationRow(const TSharedRef<FJsonObject>& CalibrationRowObject, const FImage& RowImage)
+{
+	// Create a new row to populate with data from the Json object
+	TSharedPtr<FCalibrationRow> NewRow = MakeShared<FCalibrationRow>();
+
+	// We enforce strict mode to ensure that every field in the UStruct of row data is present in the imported json.
+	// If any fields are missing, it is likely the row will be invalid, which will lead to errors in the calibration.
+	constexpr int64 CheckFlags = 0;
+	constexpr int64 SkipFlags = 0;
+	constexpr bool bStrictMode = true;
+	if (FJsonObjectConverter::JsonObjectToUStruct<FCalibrationRow>(CalibrationRowObject, NewRow.Get(), CheckFlags, SkipFlags, bStrictMode))
+	{
+		if (!RowImage.RawData.IsEmpty())
+		{
+			NewRow->MediaImage = RowImage;
+		}
+
+		Dataset.CalibrationRows.Add(NewRow);
+	}
+	else
+	{
+		UE_LOG(LogCameraCalibrationEditor, Warning, TEXT("Failed to import calibration row because at least one field could not be deserialized from the json file."));
+	}
+
+	return NewRow->Index;
+}
+
 #undef LOCTEXT_NAMESPACE
diff --git a/Engine/Plugins/VirtualProduction/CameraCalibration/Source/CameraCalibrationEditor/Private/AssetEditor/LensDistortionTool.h b/Engine/Plugins/VirtualProduction/CameraCalibration/Source/CameraCalibrationEditor/Private/AssetEditor/LensDistortionTool.h
--- a/Engine/Plugins/VirtualProduction/CameraCalibration/Source/CameraCalibrationEditor/Private/AssetEditor/LensDistortionTool.h
+++ b/Engine/Plugins/VirtualProduction/CameraCalibration/Source/CameraCalibrationEditor/Private/AssetEditor/LensDistortionTool.h
@@ -2,15 +2,12 @@
 
 #pragma once
 
-#include "CoreMinimal.h"
-
 #include "CameraCalibrationStep.h"
+
+#include "Calibrators/CameraCalibrationSolver.h"
 #include "CameraCalibrationTypes.h"
-#include "CameraLensDistortionAlgo.h"
-#include "Engine/World.h"
+#include "Engine/Texture2D.h"
 #include "ImageCore.h"
-#include "SLensDistortionToolPanel.h"
-#include "Widgets/Input/SButton.h"
 
 #include "LensDistortionTool.generated.h"
 
@@ -18,8 +15,136 @@ struct FGeometry;
 struct FPointerEvent;
 
 class FJsonObject;
-class UCameraLensDistortionAlgo;
-class ULensModel;
+class UCalibrationPointComponent;
+
+typedef UE::Tasks::TTask<FDistortionCalibrationResult> FDistortionCalibrationTask;
+
+/** List of supported calibration patterns */
+UENUM()
+enum class ECalibrationPattern : uint8
+{
+	Checkerboard,
+	Aruco,
+	Points
+};
+
+/** List of supported calibration workflows */
+UENUM()
+enum class ECalibrationType : uint8
+{
+	DistortionAndIntrinsics,
+	NodalOffset,
+	Everything
+};
+
+/** The data associated with a single captured calibration pattern/point */
+USTRUCT()
+struct FCalibrationRow
+{
+	GENERATED_BODY()
+
+	/** Index to display in list view */
+	UPROPERTY()
+	int32 Index = -1;
+
+	/** Set of captured 3D points for the calibrator in world space */
+	UPROPERTY()
+	FObjectPoints ObjectPoints;
+
+	/** Set of captured 2D pixel locations where the calibrator was detected in the image */
+	UPROPERTY()
+	FImagePoints ImagePoints;
+
+	/** Pose of the camera actor when this data for this row was captured */
+	UPROPERTY()
+	FTransform CameraPose = FTransform::Identity;
+
+	/** Pose of the calibrator actor when this data for this row was captured */
+	UPROPERTY()
+	FTransform TargetPose = FTransform::Identity;
+
+	/** The calibration pattern used to capture the data for this row */
+	UPROPERTY()
+	ECalibrationPattern Pattern = ECalibrationPattern::Checkerboard;
+
+	/** Dimensions of the detected checkerboard pattern (only valid if calibration pattern was ECalibrationPattern::Checkerboard) */
+	UPROPERTY()
+	FIntPoint CheckerboardDimensions = FIntPoint(0, 0);
+
+	/** Stored frame from the media source associated with the data for this row */
+	FImage MediaImage;
+};
+
+/** An array of captured rows with calibration data */
+USTRUCT()
+struct FCalibrationDataset
+{
+	GENERATED_BODY()
+
+	TArray<TSharedPtr<FCalibrationRow>> CalibrationRows;
+};
+
+/** Settings that control how data is captured in the tool */
+USTRUCT()
+struct FLensCaptureSettings
+{
+	GENERATED_BODY()
+
+	/** The pattern to detect in the media image */
+	UPROPERTY(EditAnywhere, Category = "Capture Settings")
+	ECalibrationPattern CalibrationPattern = ECalibrationPattern::Checkerboard;
+
+	/** An actor with Calibration Point Components that represents the virtual version of a real calibration target */
+	UPROPERTY(EditAnywhere, Category = "Capture Settings")
+	TWeakObjectPtr<AActor> Calibrator;
+
+	/** Set to true if the calibration target being used is tracked */
+	UPROPERTY(EditAnywhere, Category = "Capture Settings")
+	bool bIsCalibratorTracked = false;
+
+	/** Set to true if the camera being used is tracked */
+	UPROPERTY(EditAnywhere, Category = "Capture Settings")
+	bool bIsCameraTracked = false;
+
+	/** Display a debug overlay over the simulcam viewport showing the detected patterns that have been captured */
+	UPROPERTY(EditAnywhere, Category = "Capture Settings", meta = (EditCondition = "CalibrationPattern != ECalibrationPattern::Points", EditConditionHides))
+	bool bShowOverlay = true;
+
+	/** The name of the next calibration point to locate in the image */
+	UPROPERTY(VisibleAnywhere, Category = "Capture Settings", meta = (EditCondition = "CalibrationPattern == ECalibrationPattern::Points", EditConditionHides))
+	FText NextPoint;
+};
+
+/** Settings that control what and how the solver will compute the calibrated lens data */
+USTRUCT()
+struct FLensSolverSettings
+{
+	GENERATED_BODY()
+
+	/** Select the solver to use when solving for calibrated lens data */
+	UPROPERTY(EditAnywhere, NoClear, Category = "Solver Settings", meta = (HideViewOptions, ShowDisplayNames))
+	TSubclassOf<ULensDistortionSolver> SolverClass = ULensDistortionSolverOpenCV::StaticClass();
+
+	/** Select the data that the solver should compute and write to the Lens File */
+	UPROPERTY(EditAnywhere, Category = "Solver Settings")
+	ECalibrationType CalibrationType = ECalibrationType::DistortionAndIntrinsics;
+
+	/** An estimate for the focal length of the lens */
+	UPROPERTY(EditAnywhere, Category = "Solver Settings")
+	TOptional<double> FocalLengthGuess;
+
+	/** Set to true to prevent the solver from optimizing the focal length. The focal length guess will be used. */
+	UPROPERTY(EditAnywhere, AdvancedDisplay, Category = "Solver Settings", meta = (DisplayName = "Fix Focal Length During Optimization"))
+	bool bFixFocalLength = false;
+
+	/** Set to true to prevent the solver from optimizing the image center. The image center will be assumed to be in the exact center. */
+	UPROPERTY(EditAnywhere, AdvancedDisplay, Category = "Solver Settings", meta = (DisplayName = "Fix Image Center During Optimization"))
+	bool bFixImageCenter = false;
+
+	/** Set to true to prevent the solver from optimizing the distortion parameters. The current calibrated distortion parameters will be used. */
+	UPROPERTY(EditAnywhere, AdvancedDisplay, Category = "Solver Settings", meta = (DisplayName = "Fix Distortion During Optimization"))
+	bool bFixDistortion = false;
+};
 
 /** Data associated with a lens distortion calibration session */
 struct FLensDistortionSessionInfo
@@ -35,7 +160,7 @@ struct FLensDistortionSessionInfo
 };
 
 /**
- * ULensDistortionTool is the controller for the lens distortion panel.
+ * The primary tool used in the LensFile asset editor to capture and solve for calibrated lens data
  */
 UCLASS()
 class ULensDistortionTool : public UCameraCalibrationStep
@@ -43,7 +168,8 @@ class ULensDistortionTool : public UCameraCalibrationStep
 	GENERATED_BODY()
 
 public:
-
+	friend class SLensDistortionToolPanel;
+	
 	//~ Begin UCameraCalibrationStep interface
 	virtual void Initialize(TWeakPtr<FCameraCalibrationStepsController> InCameraCalibrationStepController) override;
 	virtual void Shutdown() override;
@@ -54,125 +180,125 @@ public:
 	virtual bool DependsOnStep(UCameraCalibrationStep* Step) const override;
 	virtual void Activate() override;
 	virtual void Deactivate() override;
-	virtual FCameraCalibrationStepsController* GetCameraCalibrationStepsController() const override;
 	virtual bool IsActive() const override;
 	virtual UMaterialInstanceDynamic* GetOverlayMID() const override;
 	virtual bool IsOverlayEnabled() const override;
 	//~ End UCameraCalibrationStep interface
 
-	/** Get the currently selected lens distortion solver class */
-	UClass* GetSolverClass();
-
-	/** Set the lens distortion solver class to use */
-	void SetSolverClass(UClass* InSolverClass);
-
-	/** Selects the algorithm by name */
-	void SetAlgo(const FName& AlgoName);
-
-	/** Resets the current algo to none */
-	void ResetAlgo();
-
-	/** Returns the currently selected algorithm */
-	UCameraLensDistortionAlgo* GetAlgo() const;
-
-	/** Returns available algorithm names */
-	TArray<FName> GetAlgos() const;
-
-	/** Called by the UI when the user wants to save the calibration data that the current algorithm is providing */
-	void OnSaveCurrentCalibrationData();
-
-	/** Triggered when the Lens Model used by the LensFile changes, allowing the tool to update its list of supported algos */
-	void OnLensModelChanged(const TSubclassOf<ULensModel>& LensModel);
-
+private:
 	/** Initiate a new calibration session (if one is not already active) */
 	void StartCalibrationSession();
 
 	/** End the current calibration session */
 	void EndCalibrationSession();
 
+	/** Deletes the .json file with the input row index that was previously exported for this session. */
+	void DeleteExportedRow(const int32& RowIndex) const;
+
 	/** Increments the session index and returns its new value */
 	uint32 AdvanceSessionRowIndex();
 
-	/** Intended to be called by one of the algo objects. Exports the data needed by the algo for the current session to a .json file on disk. */
-	void ExportSessionData(const TSharedRef<FJsonObject>& SessionDataObject);
+	/** Called by the UI when the user wants to save the calibration data */
+	void CalibrateLens();
 
-	/** Intended to be called by one of the algo objects. Exports the data associated with a single calibration row to a .json file on disk. */
-	void ExportCalibrationRow(int32 RowIndex, const TSharedRef<FJsonObject>& RowObject, const FImageView& RowImage = FImageView());
+	/** Save the results from a finished distortion calibration to the Lens File */
+	void SaveCalibrationResult();
 
-	/** Intended to be called by one of the algo objects. Deletes the .json file with the input row index that was previously exported for this session. */
-	void DeleteExportedRow(const int32& RowIndex) const;
+	/** Cancel an in-progress async calibration task */
+	void CancelCalibration();
 
-	/** Import a saved calibration data set from disk */
-	void ImportCalibrationDataset();
+	/** Get the latest status from the lens distortion calibration. Returns true if the status has changed. */
+	bool GetCalibrationStatus(FText& StatusText) const;
 
-private:
-	/** Get the directory for the current session */
-	FString GetSessionSaveDir() const;
-
-	/** Get the filename of a row with the input index */
-	FString GetRowFilename(int32 RowIndex) const;
+	/** Detects a checkerboard pattern in the input array of pixels (supplied by the media image) and fills out the 2D and 3D points in the output calibration row */
+	bool DetectCheckerboardPattern(TArray<FColor>& Pixels, FIntPoint Size, TSharedPtr<FCalibrationRow> OutRow);
 
-	/** Update the list of algos that support the current Lens Model */
-	void UpdateAlgoMap(const TSubclassOf<ULensModel>& LensModel);
+	/** Detects an aruco pattern in the input array of pixels (supplied by the media image) and fills out the 2D and 3D points in the output calibration row */
+	bool DetectArucoPattern(TArray<FColor>& Pixels, FIntPoint Size, TSharedPtr<FCalibrationRow> OutRow);
 
-	/** Build the UI for the progress window, including the window content and buttons */
-	void BuildProgressWindowWidgets();
+	/** Captures the mouse position and the 3D world location of the current calibration component and and fills out the output calibration row */
+	bool DetectPoint(const FGeometry& MyGeometry, const FPointerEvent& MouseEvent, TSharedPtr<FCalibrationRow> OutRow);
 
-	/** Called in response to user clicking "Cancel" button in the progress window to cancel a currently running calibration task */
-	FReply OnCancelPressed();
+	/** Import a saved calibration data set from disk */
+	void ImportCalibrationDataset();
 
-	/** Called in response to user clicking "Ok" button in the progress window to save the result from a finished calibration task */
-	FReply OnOkPressed();
+	/** Sets the calibrator object to be used. Updates the selection picker. */
+	void SetCalibrator(AActor* InCalibrator);
 
-	/** Save the results from a finished distortion calibration to the Lens File */
-	void SaveCalibrationResult();
+	/** Caches the calibration point components of the current calibrator actor, and updates the list if those components become invalid */
+	void UpdateCalibrationComponents();
 
-public:
-	/** Stores info about the current calibration session of this tool */
-	FLensDistortionSessionInfo SessionInfo;
+	/** Clears the list of calibration rows */
+	void ClearCalibrationRows();
 
-private:
+	/** Updates the calibration point component index (used when the calbration pattern is ECalibrationPattern::Points) */
+	void SetComponentIndex(int32 Index);
 
-	/** Pointer to the calibration steps controller */
-	TWeakPtr<FCameraCalibrationStepsController> CameraCalibrationStepsController;
+	/** Clears and redraws the overlay texture with the data from the captured calibration rows */
+	void RefreshCoverage();
 
-	/** The currently selected algorithm */
-	UPROPERTY(Transient)
-	TObjectPtr<UCameraLensDistortionAlgo> CurrentAlgo;
+	/** Adjust the 2D location of every input pixel to account for differences between the size of the captured image and the overlay texture */
+	void RescalePoints(TArray<FVector2D>& Points, FIntPoint DebugTextureSize, FIntPoint CameraFeedSize);
 
-	/** The solver class that will be used by the algos when creating a solver to calibrate for distortion */
-	TObjectPtr<UClass> SolverClass;
+	/** Get the directory for the current session */
+	FString GetSessionSaveDir() const;
 
-	/** Holds the registered camera lens distortion algos */
-	UPROPERTY(Transient)
-	TMap<FName, TSubclassOf<UCameraLensDistortionAlgo>> AlgosMap;
+	/** Get the filename of a row with the input index */
+	FString GetRowFilename(int32 RowIndex) const;
 
-	/** Holds a subset of the registered algos that support the current Lens Model */
-	UPROPERTY(Transient)
-	TMap<FName, TSubclassOf<UCameraLensDistortionAlgo>> SupportedAlgosMap;
+	/**
+	 * Import a JsonObject of calibration data that represents a single calibration row.
+	 * Returns the row index of the imported row.
+	 */
+	int32 ImportCalibrationRow(const TSharedRef<FJsonObject>& CalibrationRowObject, const FImage& RowImage);
 
-	/** Map of algo names to overlay MIDs used by those algos */
-	UPROPERTY(Transient)
-	TMap<FName, TObjectPtr<UMaterialInstanceDynamic>> AlgoOverlayMIDs;
+	/** Export global session data to a .json file */
+	void ExportSessionData();
 
-	/** True if this tool is the active one in the panel */
-	bool bIsActive = false;
+	/** Export the row data to a json file */
+	void ExportCalibrationRow(TSharedPtr<FCalibrationRow> Row);
 
-	/** UI Widget for this Tool */
+private:
+	/** UI Widget for this tool */
 	TSharedPtr<SLensDistortionToolPanel> DistortionWidget;
 
-	/** UI Widget for the calibration task progress window */
-	TSharedPtr<SWindow> ProgressWindow;
+	FLensSolverSettings SolverSettings;
+	FLensCaptureSettings CaptureSettings;
 
-	/** UI Widget for the text to be displayed in the progress window */
-	TSharedPtr<STextBlock> ProgressTextWidget;
+	/** Collection of calibration rows containing the data needed to run a lens calibration */
+	FCalibrationDataset Dataset;
 
-	/** UI Widget for the "Ok" button in the progress window */
-	TSharedPtr<SButton> OkayButton;
+	/** The solver object that runs the lens calibration on another thread. The ptr is stored to access the solver to get status updates and/or cancel the calibration from the game thread */
+	TObjectPtr<ULensDistortionSolver> Solver;
 
 	/** An asynchronous task handle. When valid, the tool will poll its state to determine when the task has completed, and then extract the calibration result from this task handle. */
 	FDistortionCalibrationTask CalibrationTask;
 
 	/** The result from the most recently completed distortion calibration. It is saved as a member variable because the result is not saved immediately, but only if the user confirms that it should be saved. */
 	FDistortionCalibrationResult CalibrationResult;
+
+	/** Container for the set of calibrator components selected in the component combobox */
+	TArray<TWeakObjectPtr<UCalibrationPointComponent>> CalibrationComponents;
+
+	/** 
+	 * The 3D world locations of the calibration point components should all come from the same frame.
+	 * Therefore, the locations are all saved at once when the first calibration point is selected (only applicable to ECalibrationPattern::Points)
+	 */
+	TArray<FVector> CachedComponentLocations;
+
+	/** Index of the next calibration point component to capture (only applicable to ECalibrationPattern::Points) */
+	int32 CalibrationComponentIndex = -1;
+
+	/** Capture session info, used to track whether newly captured data should be added to an existing dataset */
+	FLensDistortionSessionInfo SessionInfo;
+
+	/** True if this tool is the active tab in the UI */
+	bool bIsActive = false;
+
+	/** Material and Texture that draw the coverage overlay on top of the simulcam viewport during capture */
+	TObjectPtr<UMaterialInstanceDynamic> OverlayMID;
+	TObjectPtr<UTexture2D> OverlayTexture;
+
+	/** Weak ptr to the steps controller that created this step */
+	TWeakPtr<FCameraCalibrationStepsController> WeakStepsController;
 };
diff --git a/Engine/Plugins/VirtualProduction/CameraCalibration/Source/CameraCalibrationEditor/Private/AssetEditor/SLensDistortionToolPanel.cpp b/Engine/Plugins/VirtualProduction/CameraCalibration/Source/CameraCalibrationEditor/Private/AssetEditor/SLensDistortionToolPanel.cpp
--- a/Engine/Plugins/VirtualProduction/CameraCalibration/Source/CameraCalibrationEditor/Private/AssetEditor/SLensDistortionToolPanel.cpp
+++ b/Engine/Plugins/VirtualProduction/CameraCalibration/Source/CameraCalibrationEditor/Private/AssetEditor/SLensDistortionToolPanel.cpp
@@ -2,380 +2,507 @@
 
 #include "SLensDistortionToolPanel.h"
 
-#include "AssetRegistry/AssetData.h"
-#include "Calibrators/CameraCalibrationSolver.h"
-#include "CameraCalibrationSettings.h"
-#include "CameraCalibrationSubsystem.h"
-#include "CameraLensDistortionAlgo.h"
-#include "Dialog/SCustomDialog.h"
-#include "EditorFontGlyphs.h"
-#include "Engine/Selection.h"
-#include "Framework/MultiBox/MultiBoxBuilder.h"
-#include "LensDistortionTool.h"
-#include "LensFile.h"
-#include "Styling/AppStyle.h"
-#include "Styling/CoreStyle.h"
-#include "Styling/SlateStyle.h"
+#include "CalibrationPointComponent.h"
+#include "DetailLayoutBuilder.h"
+#include "Framework/Application/SlateApplication.h"
+#include "ImageUtils.h"
+#include "IStructureDetailsView.h"
+#include "Modules/ModuleManager.h"
 #include "UI/CameraCalibrationWidgetHelpers.h"
-#include "Widgets/Images/SImage.h"
+#include "UI/SFilterableActorPicker.h"
+#include "UI/SImageTexture.h"
 #include "Widgets/Input/SButton.h"
-#include "Widgets/Input/SComboBox.h"
 #include "Widgets/SBoxPanel.h"
-#include "Widgets/Text/STextBlock.h"
-#include "Widgets/Views/STreeView.h"
+#include "Widgets/SWindow.h"
 
-#define LOCTEXT_NAMESPACE "LensDistortionTool"
+#define LOCTEXT_NAMESPACE "SLensDistortionToolPanel"
 
+void SCalibrationDatasetRow::Construct(const FArguments& Args, const TSharedRef<STableViewBase>& OwnerTableView, TSharedPtr<FCalibrationRow>& InRowData)
+{
+	RowData = InRowData;
+
+	FSuperRowType::FArguments StyleArguments = FSuperRowType::FArguments()
+		.Padding(1.0f)
+		.Style(&FAppStyle().GetWidgetStyle<FTableRowStyle>("TableView.AlternatingRow"));
+
+	SMultiColumnTableRow<TSharedPtr<FCalibrationRow>>::Construct(StyleArguments, OwnerTableView);
+}
+
+TSharedRef<SWidget> SCalibrationDatasetRow::GenerateWidgetForColumn(const FName& ColumnName)
+{
+	if (ColumnName == TEXT("Index"))
+	{
+		return SNew(SVerticalBox)
+			+ SVerticalBox::Slot()
+			.HAlign(EHorizontalAlignment::HAlign_Center)
+			.VAlign(EVerticalAlignment::VAlign_Center)
+			[SNew(STextBlock).Text(FText::AsNumber(RowData->Index))];
+	}
+	else if (ColumnName == TEXT("Image"))
+	{
+		// Generate a transient thumbnail texture to display in the tool
+		const FImage& MediaImage = RowData->MediaImage;
+
+		if ((MediaImage.SizeX < 1) || (MediaImage.SizeY < 1))
+		{
+			const FString Text = FString::Printf(TEXT("Image Unavailable"));
+			return SNew(STextBlock).Text(FText::FromString(Text));
+		}
+
+		constexpr int32 ResolutionDivider = 4;
+		const FIntPoint ThumbnailSize = FIntPoint(MediaImage.SizeX / ResolutionDivider, MediaImage.SizeY / ResolutionDivider);
+
+		FImage ThumbnailImage;
+		FImageCore::ResizeTo(MediaImage, ThumbnailImage, ThumbnailSize.X, ThumbnailSize.Y, MediaImage.Format, MediaImage.GetGammaSpace());
+
+		if (UTexture2D* Thumbnail = FImageUtils::CreateTexture2DFromImage(ThumbnailImage))
+		{
+			return SNew(SImageTexture, Thumbnail)
+				.MinDesiredHeight(4 * FCameraCalibrationWidgetHelpers::DefaultRowHeight)
+				.MaxDesiredHeight(4 * FCameraCalibrationWidgetHelpers::DefaultRowHeight);
+		}
+		else
+		{
+			const FString Text = FString::Printf(TEXT("Image Unavailable"));
+			return SNew(STextBlock).Text(FText::FromString(Text));
+		}
+	}
+	else if (ColumnName == TEXT("ImagePoint"))
+	{
+		const FString Text = FString::Printf(TEXT("(%.2f, %.2f)"),
+			RowData->ImagePoints.Points[0].X,
+			RowData->ImagePoints.Points[0].Y);
+
+		return SNew(STextBlock).Text(FText::FromString(Text));
+	}
+	else if (ColumnName == TEXT("ObjectPoint"))
+	{
+		const FString Text = FString::Printf(TEXT("(%.2f, %.2f, %.2f)"),
+			RowData->ObjectPoints.Points[0].X,
+			RowData->ObjectPoints.Points[0].Y,
+			RowData->ObjectPoints.Points[0].Z);
+
+		return SNew(STextBlock).Text(FText::FromString(Text));
+	}
+
+	return SNullWidget::NullWidget;
+}
+
+void FCaptureSettingsCustomization::CustomizeDetails(IDetailLayoutBuilder& DetailBuilder)
+{
+	PropertyHandle = DetailBuilder.GetProperty(GET_MEMBER_NAME_CHECKED(FLensCaptureSettings, Calibrator));
+
+	IDetailPropertyRow* CalibratorRow = DetailBuilder.EditDefaultProperty(PropertyHandle);
+	CalibratorRow->CustomWidget()
+		.NameContent()
+		[
+			PropertyHandle->CreatePropertyNameWidget()
+		]
+		.ValueContent()
+		[
+			SNew(SFilterableActorPicker)
+				.OnSetObject(this, &FCaptureSettingsCustomization::OnCalibratorSelected)
+				.OnShouldFilterAsset(this, &FCaptureSettingsCustomization::DoesAssetHaveCalibrationComponent)
+				.ActorAssetData(this, &FCaptureSettingsCustomization::GetCalibratorAssetData)
+		];
+
+	NextPointPropertyHandle = DetailBuilder.GetProperty(GET_MEMBER_NAME_CHECKED(FLensCaptureSettings, NextPoint));
+	IDetailPropertyRow* NextPointRow = DetailBuilder.EditDefaultProperty(NextPointPropertyHandle);
+	NextPointRow->CustomWidget()
+		.NameContent()
+		[
+			NextPointPropertyHandle->CreatePropertyNameWidget()
+		]
+		.ValueContent()
+		[
+			SNew(STextBlock)
+				.Text(this, &FCaptureSettingsCustomization::GetNextPointName)
+				.Font(DetailBuilder.GetDetailFont())
+		];
+}
+
+FText FCaptureSettingsCustomization::GetNextPointName() const
+{
+	FText Name;
+	NextPointPropertyHandle->GetValue(Name);
+	return Name;
+}
 
-void SLensDistortionToolPanel::Construct(const FArguments& InArgs, ULensDistortionTool* InTool)
+void FCaptureSettingsCustomization::OnCalibratorSelected(const FAssetData& AssetData)
+{
+	if (AssetData.IsValid())
+	{
+		PropertyHandle->SetValue(AssetData);
+	}
+}
+
+bool FCaptureSettingsCustomization::DoesAssetHaveCalibrationComponent(const FAssetData& AssetData) const
+{
+	if (const AActor* Actor = Cast<AActor>(AssetData.GetAsset()))
+	{
+		constexpr uint32 NumInlineAllocations = 32;
+		TArray<UCalibrationPointComponent*, TInlineAllocator<NumInlineAllocations>> CalibrationPoints;
+		Actor->GetComponents(CalibrationPoints);
+
+		return (CalibrationPoints.Num() > 0);
+	}
+
+	return false;
+}
+
+FAssetData FCaptureSettingsCustomization::GetCalibratorAssetData() const
+{
+	FAssetData Tmp;
+	PropertyHandle->GetValue(Tmp);
+	return Tmp;
+}
+
+void SLensDistortionToolPanel::Construct(const FArguments& InArgs, ULensDistortionTool* InTool, TWeakPtr<FCameraCalibrationStepsController> InStepsController)
 {
 	Tool = InTool;
+	WeakStepsController = InStepsController;
+
+	BuildProgressWindow();
+
+	FPropertyEditorModule& PropertyEditor = FModuleManager::Get().LoadModuleChecked<FPropertyEditorModule>(TEXT("PropertyEditor"));
+
+	FStructureDetailsViewArgs DefaultStructuDetailsViewArgs;
+	FDetailsViewArgs DefaultDetailsViewArgs;
+	DefaultDetailsViewArgs.bAllowSearch = false;
 
-	// This will be the widget wrapper of the custom algo UI.
-	UI = SNew(SVerticalBox);
+	TSharedRef<FStructOnScope> SolverSettingsStruct = MakeShared<FStructOnScope>(FLensSolverSettings::StaticStruct(), reinterpret_cast<uint8*>(&Tool->SolverSettings));
+	TSharedPtr<IStructureDetailsView> SolverSettingsDetailsView = PropertyEditor.CreateStructureDetailView(DefaultDetailsViewArgs, DefaultStructuDetailsViewArgs, SolverSettingsStruct);
+
+	TSharedRef<FStructOnScope> CaptureSettingsStruct = MakeShared<FStructOnScope>(FLensCaptureSettings::StaticStruct(), reinterpret_cast<uint8*>(&Tool->CaptureSettings));
+
+	TSharedPtr<IStructureDetailsView> CaptureSettingsDetailsView = PropertyEditor.CreateStructureDetailView(DefaultDetailsViewArgs, DefaultStructuDetailsViewArgs, TSharedPtr<FStructOnScope>());
+	CaptureSettingsDetailsView->GetDetailsView()->OnFinishedChangingProperties().AddSP(this, &SLensDistortionToolPanel::OnCaptureSettingsChanged);
+
+	CaptureSettingsDetailsView->GetDetailsView()->RegisterInstancedCustomPropertyLayout(FLensCaptureSettings::StaticStruct(),
+		FOnGetDetailCustomizationInstance::CreateLambda([]() { return MakeShared<FCaptureSettingsCustomization>(); }));
+
+	CaptureSettingsDetailsView->SetStructureData(CaptureSettingsStruct);
 
 	ChildSlot
 	[
 		SNew(SHorizontalBox)
 
-		+ SHorizontalBox::Slot() // Right toolbar
+		+ SHorizontalBox::Slot()
 		.FillWidth(0.25f)
-		[ 
+		[
 			SNew(SVerticalBox)
 
-			+ SVerticalBox::Slot() // Solver picker
-			.MaxHeight(FCameraCalibrationWidgetHelpers::DefaultRowHeight)
-			[FCameraCalibrationWidgetHelpers::BuildLabelWidgetPair(LOCTEXT("SolverPickerWidgetText", "Lens Distortion Solver"), BuildSolverPickerWidget())]
+			+ SVerticalBox::Slot()
+			.AutoHeight()
+			[
+				CaptureSettingsDetailsView->GetWidget().ToSharedRef()
+			]
 
-			+ SVerticalBox::Slot() // Algo picker
-			.MaxHeight(FCameraCalibrationWidgetHelpers::DefaultRowHeight)
-			[ FCameraCalibrationWidgetHelpers::BuildLabelWidgetPair(LOCTEXT("LensDistortionAlgo", "Lens Distortion Algo"), BuildAlgoPickerWidget())]
+			+ SVerticalBox::Slot()
+			.AutoHeight()
+			[
+				SolverSettingsDetailsView->GetWidget().ToSharedRef()
+			]
+
+			+ SVerticalBox::Slot()
+			.AutoHeight()
+			.MaxHeight(12 * FCameraCalibrationWidgetHelpers::DefaultRowHeight)
+			[
+				BuildDatasetListView()
+			]
 
-			+ SVerticalBox::Slot() // Algo UI
+			+ SVerticalBox::Slot()
 			.AutoHeight()
-			[ BuildUIWrapper() ]
+			.HAlign(HAlign_Center)
+			.Padding(0, 20)
+			[
+				SNew(SButton)
+				.Text(LOCTEXT("ClearAll", "Clear All"))
+				.HAlign(HAlign_Center)
+				.VAlign(VAlign_Center)
+				.OnClicked(this, &SLensDistortionToolPanel::OnClearCalibrationRowsClicked)
+			]
 
-			+ SVerticalBox::Slot() // Import dataset
+			+ SVerticalBox::Slot()
 			.AutoHeight()
 			.HAlign(HAlign_Center)
 			.Padding(0, 20)
 			[
-				SNew(SHorizontalBox)
-
-				+ SHorizontalBox::Slot()
-				.AutoWidth()
-				[
-					SNew(SButton).Text(LOCTEXT("ImportCalibrationDataset", "Import Lens Distortion Dataset"))
-					.HAlign(HAlign_Center)
-					.VAlign(VAlign_Center)
-					.OnClicked_Lambda([this]() -> FReply
-					{
-						if (ULensDistortionTool* LensDistortionTool = Tool.Get())
-						{
-							LensDistortionTool->ImportCalibrationDataset();
-
-							// After importing, the tool may have switched the active algo, so redraw the UI accordingly
-							UCameraLensDistortionAlgo* Algo = LensDistortionTool->GetAlgo();
-							UI->ClearChildren();
-							UI->AddSlot()[Algo->BuildUI()];
-
-							for (const TSharedPtr<FString>& AlgoString : CurrentAlgos)
-							{
-								if (AlgoString->Equals(Algo->FriendlyName().ToString()))
-								{
-									AlgosComboBox->SetSelectedItem(AlgoString);
-									break;
-								}
-							}
-						}
-						return FReply::Handled();
-					})
-					.Visibility_Lambda([]() -> EVisibility
-					{
-						if (GetDefault<UCameraCalibrationSettings>()->IsCalibrationDatasetImportExportEnabled())
-						{
-							return EVisibility::Visible;
-						}
-						return EVisibility::Collapsed;
-					})
-				]
+				SNew(SButton)
+				.Text(LOCTEXT("ImportDataset", "Import Dataset"))
+				.HAlign(HAlign_Center)
+				.VAlign(VAlign_Center)
+				.OnClicked(this, &SLensDistortionToolPanel::OnImportDatasetClicked)
 			]
 
-			+ SVerticalBox::Slot() // Save Offset
+			+ SVerticalBox::Slot()
 			.AutoHeight()
+			.HAlign(HAlign_Center)
 			.Padding(0, 20)
 			[
-				SNew(SButton).Text(LOCTEXT("AddToLUT", "Add To Lens Distortion Calibration"))
+				SNew(SButton)
+				.Text(LOCTEXT("CalibrateLens", "Calibrate Lens"))
 				.HAlign(HAlign_Center)
 				.VAlign(VAlign_Center)
-				.OnClicked_Lambda([WeakTool = Tool]() -> FReply
-				{
-					if (WeakTool.IsValid())
-					{
-						WeakTool->OnSaveCurrentCalibrationData();
-					}
-					return FReply::Handled();
-				})
+				.OnClicked(this, &SLensDistortionToolPanel::OnCalibrateClicked)
 			]
 		]
 	];
 }
 
-TSharedRef<SWidget> SLensDistortionToolPanel::BuildUIWrapper()
+void SLensDistortionToolPanel::OpenProgressWindow()
 {
-	return SNew(SVerticalBox)
+	ProgressTextWidget->SetText(LOCTEXT("CalibrationProgressText", "Calibrating..."));
 
-		+ SVerticalBox::Slot() // Algo's Widget
-		.AutoHeight()
-		[ UI.ToSharedRef() ];
+	// The okay button will be disabled until the calibration is complete. 
+	OkayButton->SetEnabled(false);
+	ProgressWindow->ShowWindow();
 }
 
-void SLensDistortionToolPanel::UpdateUI()
+void SLensDistortionToolPanel::MarkProgressFinished()
 {
-	check(AlgosComboBox.IsValid());
-
-	// Get current algo to later compare with new one
-	const UCameraLensDistortionAlgo* OldAlgo = Tool->GetAlgo();
+	OkayButton->SetEnabled(true);
+}
 
-	// Set new algo by name
-	TSharedPtr<FString> SelectedAlgo = AlgosComboBox->GetSelectedItem();
-	if (!SelectedAlgo)
-	{
-		Tool->ResetAlgo();
-	}
-	else
-	{
-		const FName AlgoName(*SelectedAlgo);
-		Tool->SetAlgo(AlgoName);
-	}
+void SLensDistortionToolPanel::UpdateProgressText(const FText& ProgressText)
+{
+	ProgressTextWidget->SetText(ProgressText);
+}
 
-	// Get the new algo
-	UCameraLensDistortionAlgo* Algo = Tool->GetAlgo();
+void SLensDistortionToolPanel::Shutdown()
+{
+	ProgressWindow->HideWindow();
+}
 
-	// If we didn't change the algo, we're done here.
-	if (Algo == OldAlgo)
+void SLensDistortionToolPanel::OnCaptureSettingsChanged(const FPropertyChangedEvent& PropertyChangedEvent)
+{
+	if (!Tool.IsValid())
 	{
 		return;
 	}
 
-	// Remove old UI
-	check(UI.IsValid());
-	UI->ClearChildren();
+	const FName PropertyName = (PropertyChangedEvent.Property != NULL) ? PropertyChangedEvent.Property->GetFName() : NAME_None;
 
-	// nullptr may indicate that it was unregistered, so refresh combobox options.
-	if (!Algo)
+	if (PropertyName == GET_MEMBER_NAME_CHECKED(FLensCaptureSettings, bShowOverlay))
 	{
-		UpdateAlgosOptions();
-		return;
+		if (TSharedPtr<FCameraCalibrationStepsController> StepsController = WeakStepsController.Pin())
+		{
+			StepsController->SetOverlayEnabled(Tool->CaptureSettings.bShowOverlay);
+		}
+	}
+	else if (PropertyName == GET_MEMBER_NAME_CHECKED(FLensCaptureSettings, Calibrator))
+	{
+		Tool->SetCalibrator(Tool->CaptureSettings.Calibrator.Get());
+	}
+	else if (PropertyName == GET_MEMBER_NAME_CHECKED(FLensCaptureSettings, CalibrationPattern))
+	{
+		DatasetListHeader->SetShowGeneratedColumn("Image", (Tool->CaptureSettings.CalibrationPattern != ECalibrationPattern::Points));
+		DatasetListHeader->SetShowGeneratedColumn("ImagePoint", (Tool->CaptureSettings.CalibrationPattern == ECalibrationPattern::Points));
+		DatasetListHeader->SetShowGeneratedColumn("ObjectPoint", (Tool->CaptureSettings.CalibrationPattern == ECalibrationPattern::Points));
 	}
-
-	// Assign GUI
-	UI->AddSlot() [Algo->BuildUI()];
 }
 
-void SLensDistortionToolPanel::UpdateAlgosOptions()
+FReply SLensDistortionToolPanel::OnClearCalibrationRowsClicked()
 {
-	CurrentAlgos.Empty();
-
-	// Ask the Tool for the list of registered Algos
-
-	if (!Tool.IsValid())
+	if (Tool.IsValid())
 	{
-		return;
+		Tool->ClearCalibrationRows();
 	}
+	return FReply::Handled();
+}
 
-	for (FName& AlgoName : Tool->GetAlgos())
+FReply SLensDistortionToolPanel::OnImportDatasetClicked()
+{
+	if (Tool.IsValid())
 	{
-		CurrentAlgos.Add(MakeShared<FString>(AlgoName.ToString()));
+		Tool->ImportCalibrationDataset();
 	}
+	return FReply::Handled();
+}
 
-	// Pick the first option by default (if available)
-	if (CurrentAlgos.Num())
-	{
-		AlgosComboBox->SetSelectedItem(CurrentAlgos[0]);
-	}
-	else
+FReply SLensDistortionToolPanel::OnCalibrateClicked()
+{
+	if (Tool.IsValid())
 	{
-		AlgosComboBox->SetSelectedItem(nullptr);
+		Tool->CalibrateLens();
 	}
-
-	// Ask the ComboBox to refresh its options from its source (that we just updated)
-	AlgosComboBox->RefreshOptions();
+	return FReply::Handled();
 }
 
-TSharedRef<SWidget> SLensDistortionToolPanel::BuildSolverPickerWidget()
+TSharedRef<SWidget> SLensDistortionToolPanel::BuildDatasetListView()
 {
-	return SNew(SHorizontalBox)
+	if (!Tool.IsValid())
+	{
+		return SNullWidget::NullWidget;
+	}
 
-		+ SHorizontalBox::Slot()
-		.AutoWidth()
-		.Padding(FMargin(0.0f, 0.0f, 0.0f, 0.0f))
-		[
-			SNew(SComboButton)
-			.OnGetMenuContent(this, &SLensDistortionToolPanel::BuildSolverList)
-			.ContentPadding(FMargin(4.0, 2.0))
-			.ButtonContent()
-			[
-				SNew(STextBlock).Text(this, &SLensDistortionToolPanel::GetSelectedSolverName)
-			]
-		];
-}
+	TArray<FName> HiddenColumns = { "ImagePoint", "ObjectPoint" };
 
-TSharedRef<SWidget> SLensDistortionToolPanel::BuildSolverList() 
-{
-	// Generate menu
-	FMenuBuilder MenuBuilder(true, nullptr);
-	MenuBuilder.BeginSection("DistortionSolverSection", LOCTEXT("DistortionSolverSectionText", "Distortion Solvers"));
+	DatasetListHeader = SNew(SHeaderRow)
+		.HiddenColumnsList(HiddenColumns)
 
-	TArray<UClass*> DerivedSolverClasses;
-	GetDerivedClasses(ULensDistortionSolver::StaticClass(), DerivedSolverClasses);
+		+ SHeaderRow::Column("Index")
+		.DefaultLabel(LOCTEXT("IndexHeaderLabel", "Index"))
+		.FillWidth(0.2f)
 
-	if (!DerivedSolverClasses.IsEmpty())
-	{
-		for (UClass* SolverClass : DerivedSolverClasses)
-		{
-			ULensDistortionSolver* SolverCDO = Cast<ULensDistortionSolver>(SolverClass->GetDefaultObject());
-
-			if (SolverCDO->IsEnabled())
-			{
-				MenuBuilder.AddMenuEntry(
-					SolverCDO->GetDisplayName(),
-					SolverCDO->GetDisplayName(), // TODO: If this remains a user-facing option, then it may be better to have the solver implement a class description text field to use as the tooltip instead
-					FSlateIcon(),
-					FUIAction(
-						FExecuteAction::CreateSP(this, &SLensDistortionToolPanel::OnSolverClassSelected, SolverClass),
-						FCanExecuteAction(),
-						FIsActionChecked::CreateSP(this, &SLensDistortionToolPanel::IsSolverClassSelected, SolverClass)
-					),
-					NAME_None,
-					EUserInterfaceActionType::RadioButton
-				);
-			}
-		}
-	}
+		+ SHeaderRow::Column("Image")
+		.DefaultLabel(LOCTEXT("ImageHeaderLabel", "Image"))
+		.FillWidth(0.8f)
 
-	MenuBuilder.EndSection();
+		+ SHeaderRow::Column("ImagePoint")
+		.DefaultLabel(LOCTEXT("ImagePointHeaderLabel", "Pixel Location"))
+		.FillWidth(0.4f)
 
-	return MenuBuilder.MakeWidget();
+		+ SHeaderRow::Column("ObjectPoint")
+		.DefaultLabel(LOCTEXT("ObjectPointHeaderLabel", "World Position"))
+		.FillWidth(0.4f);
 
-}
+	DatasetListView = SNew(SListView<TSharedPtr<FCalibrationRow>>)
+		.ListItemsSource(&Tool->Dataset.CalibrationRows)
+		.SelectionMode(ESelectionMode::Multi)
+		.HeaderRow(DatasetListHeader.ToSharedRef())
+		.OnGenerateRow(this, &SLensDistortionToolPanel::OnGenerateDatasetRow)
+		.OnKeyDownHandler(this, &SLensDistortionToolPanel::OnDatasetRowKeyPressed);
 
-void SLensDistortionToolPanel::OnSolverClassSelected(UClass* SolverClass)
-{
-	Tool->SetSolverClass(SolverClass);
+	return DatasetListView.ToSharedRef();
 }
 
-bool SLensDistortionToolPanel::IsSolverClassSelected(UClass* SolverClass) const
+TSharedRef<ITableRow> SLensDistortionToolPanel::OnGenerateDatasetRow(TSharedPtr<FCalibrationRow> InItem, const TSharedRef<STableViewBase>& OwnerTable)
 {
-	return (Tool->GetSolverClass() == SolverClass);
+	return SNew(SCalibrationDatasetRow, OwnerTable, InItem);
 }
 
-FText SLensDistortionToolPanel::GetSelectedSolverName() const
+FReply SLensDistortionToolPanel::OnDatasetRowKeyPressed(const FGeometry& Geometry, const FKeyEvent& KeyEvent)
 {
-	if (UClass* SolverClass = Tool->GetSolverClass())
+	if (!Tool.IsValid())
 	{
-		ULensDistortionSolver* SolverCDO = Cast<ULensDistortionSolver>(SolverClass->GetDefaultObject());
-		return SolverCDO->GetDisplayName();
+		return FReply::Unhandled();
 	}
-	
-	return FText::GetEmpty();
-}
 
-TSharedRef<SWidget> SLensDistortionToolPanel::BuildAlgoPickerWidget()
-{
-	// Create ComboBox widget
+	if ((KeyEvent.GetKey() == EKeys::A) && KeyEvent.GetModifierKeys().IsControlDown())
+	{
+		// Select all items
+		DatasetListView->SetItemSelection(Tool->Dataset.CalibrationRows, true);
+		return FReply::Handled();
+	}
+	else if (KeyEvent.GetKey() == EKeys::Escape)
+	{
+		// De-select all items
+		DatasetListView->ClearSelection();
+		return FReply::Handled();
+	}
+	else if (KeyEvent.GetKey() == EKeys::Delete)
+	{
+		// Delete selected items
+		const TArray<TSharedPtr<FCalibrationRow>> SelectedItems = DatasetListView->GetSelectedItems();
 
-	AlgosComboBox = SNew(SComboBox<TSharedPtr<FString>>)
-		.OptionsSource(&CurrentAlgos)
-		.OnSelectionChanged_Lambda([&](TSharedPtr<FString> NewValue, ESelectInfo::Type Type) -> void
-		{
-			// Replace the custom algo widget
-			UpdateUI();
-		})
-		.OnGenerateWidget_Lambda([&](TSharedPtr<FString> InOption) -> TSharedRef<SWidget>
+		for (const TSharedPtr<FCalibrationRow>& SelectedItem : SelectedItems)
 		{
-			return SNew(STextBlock).Text(FText::FromString(*InOption));
-		})
-		.InitiallySelectedItem(nullptr)
-		[
-			SNew(STextBlock)
-			.Text_Lambda([&]() -> FText
-			{
-				if (AlgosComboBox.IsValid() && AlgosComboBox->GetSelectedItem().IsValid())
-				{
-					return FText::FromString(*AlgosComboBox->GetSelectedItem());
-				}
-
-				return LOCTEXT("InvalidComboOption", "Invalid");
-			})
-		];
+			Tool->Dataset.CalibrationRows.Remove(SelectedItem);
+			Tool->DeleteExportedRow(SelectedItem->Index);
+		}
 
-	// Update the object holding this combobox's options source
-	UpdateAlgosOptions();
+		Tool->RefreshCoverage();
 
-	// Pick the first option by default (if available)
-	if (CurrentAlgos.Num())
-	{
-		AlgosComboBox->SetSelectedItem(CurrentAlgos[0]);
+		return FReply::Handled();
 	}
-	else
+
+	return FReply::Unhandled();
+}
+
+void SLensDistortionToolPanel::RefreshListView()
+{
+	if (DatasetListView)
 	{
-		AlgosComboBox->SetSelectedItem(nullptr);
+		DatasetListView->RequestListRefresh();
 	}
+}
 
-	return SNew(SHorizontalBox)
-
-		+ SHorizontalBox::Slot() // algo picker
+void SLensDistortionToolPanel::BuildProgressWindow()
+{
+	ProgressWindow = SNew(SWindow)
+		.Title(LOCTEXT("ProgressWindowTitle", "Distortion Calibration Progress"))
+		.SizingRule(ESizingRule::Autosized)
+		.IsTopmostWindow(true)
+		.HasCloseButton(false)
+		.SupportsMaximize(false)
+		.SupportsMinimize(true);
+
+	ProgressTextWidget = SNew(STextBlock).Text(FText::GetEmpty());
+
+	OkayButton = SNew(SButton)
+		.IsEnabled(false)
+		.HAlign(HAlign_Center)
+		.VAlign(VAlign_Center)
+		.Text(LOCTEXT("OkText", "Ok"))
+		.OnClicked(this, &SLensDistortionToolPanel::OnOkPressed);
+
+	TSharedRef<SWidget> WindowContent = SNew(SVerticalBox)
+
+		// Text widget to display the current progress of the calibration
+		+ SVerticalBox::Slot()
+		.HAlign(EHorizontalAlignment::HAlign_Center)
+		.VAlign(EVerticalAlignment::VAlign_Center)
 		[
-			AlgosComboBox.ToSharedRef()
+			ProgressTextWidget.ToSharedRef()
 		]
 
-		+ SHorizontalBox::Slot() // Help button
-		.AutoWidth()
+		// Ok and Cancel buttons
+		+ SVerticalBox::Slot()
+		.HAlign(EHorizontalAlignment::HAlign_Center)
+		.VAlign(EVerticalAlignment::VAlign_Center)
 		[
-			SNew(SButton)
-			.ToolTipText(LOCTEXT("ShowHelp_Tip", "Help about this algo"))
-			.HAlign(HAlign_Center)
-			.VAlign(VAlign_Center)
-			.ButtonStyle(FAppStyle::Get(), "HoverHintOnly")
-			.OnClicked_Lambda([&]() -> FReply
-			{
-				if (!Tool.IsValid())
-				{
-					return FReply::Handled();
-				}
-
-				UCameraLensDistortionAlgo* Algo = Tool->GetAlgo();
-
-				if (!Algo)
-				{
-					return FReply::Handled();
-				}
-
-				TSharedRef< SCustomDialog> AlgoHelpWindow = 
-					SNew(SCustomDialog)
-					.Title(FText::FromName(Tool->FriendlyName()))
-					.Content()
-					[
-						Algo->BuildHelpWidget()
-					]
-					.Buttons({
-						SCustomDialog::FButton(LOCTEXT("Ok", "Ok")),
-					});
-
-				AlgoHelpWindow->Show();
-
-				return FReply::Handled();
-			})
+			SNew(SHorizontalBox)
+
+			+ SHorizontalBox::Slot()
+			.AutoWidth()
 			[
-				SNew(STextBlock)
-				.Font(FAppStyle::Get().GetFontStyle("FontAwesome.12"))
-				.Text(FEditorFontGlyphs::Info_Circle)
-				.ColorAndOpacity(FLinearColor::White)
+				OkayButton.ToSharedRef()
 			]
-		]
-		;
+
+			+ SHorizontalBox::Slot()
+			.AutoWidth()
+			[
+				SNew(SButton)
+				.Text(LOCTEXT("CancelText", "Cancel"))
+				.HAlign(HAlign_Center)
+				.VAlign(VAlign_Center)
+				.OnClicked(this, &SLensDistortionToolPanel::OnCancelPressed)
+			]
+		];
+
+	ProgressWindow->SetContent(WindowContent);
+
+	// Create the window, but start with it hidden. When the user initiates a calibration, the progress window will be shown.
+	FSlateApplication::Get().AddWindow(ProgressWindow.ToSharedRef());
+	ProgressWindow->HideWindow();
+}
+
+FReply SLensDistortionToolPanel::OnCancelPressed()
+{
+	if (Tool.IsValid())
+	{
+		Tool->CancelCalibration();
+		Tool->CalibrationTask = {};
+	}
+
+	ProgressWindow->HideWindow();
+	SetEnabled(true);
+
+	return FReply::Handled();
 }
 
+FReply SLensDistortionToolPanel::OnOkPressed()
+{
+	if (Tool.IsValid())
+	{
+		Tool->SaveCalibrationResult();
+	}
+
+	ProgressWindow->HideWindow();
+	SetEnabled(true);
+
+	return FReply::Handled();
+}
 
 #undef LOCTEXT_NAMESPACE
diff --git a/Engine/Plugins/VirtualProduction/CameraCalibration/Source/CameraCalibrationEditor/Private/AssetEditor/SLensDistortionToolPanel.h b/Engine/Plugins/VirtualProduction/CameraCalibration/Source/CameraCalibrationEditor/Private/AssetEditor/SLensDistortionToolPanel.h
--- a/Engine/Plugins/VirtualProduction/CameraCalibration/Source/CameraCalibrationEditor/Private/AssetEditor/SLensDistortionToolPanel.h
+++ b/Engine/Plugins/VirtualProduction/CameraCalibration/Source/CameraCalibrationEditor/Private/AssetEditor/SLensDistortionToolPanel.h
@@ -1,74 +1,103 @@
 // Copyright Epic Games, Inc. All Rights Reserved.
 #pragma once
 
-#include "CoreMinimal.h"
-
-#include "Widgets/DeclarativeSyntaxSupport.h"
 #include "Widgets/SCompoundWidget.h"
-#include "UObject/StrongObjectPtr.h"
 
-class FString;
-class ULensFile;
+#include "CameraCalibrationStepsController.h"
+#include "IDetailCustomization.h"
+#include "LensDistortionTool.h"
+#include "PropertyHandle.h"
+#include "Widgets/SWindow.h"
+#include "Widgets/Text/STextBlock.h"
+#include "Widgets/Views/SHeaderRow.h"
+#include "Widgets/Views/SListView.h"
+#include "Widgets/Views/STableRow.h"
+
 class ULensDistortionTool;
 
-template<typename OptionType>
-class SComboBox;
+struct FCalibrationRow;
 
-struct FArguments;
+/** Customization for the FCaptureSettings struct */
+class FCaptureSettingsCustomization : public IDetailCustomization
+{
+public:
+	virtual void CustomizeDetails(IDetailLayoutBuilder& DetailBuilder) override;
 
-/**
- * Wrapper UI for the specified calibration step.
- */
-class SLensDistortionToolPanel : public SCompoundWidget
+private:
+	TSharedPtr<IPropertyHandle> PropertyHandle;
+	TSharedPtr<IPropertyHandle> NextPointPropertyHandle;
+
+	FText GetNextPointName() const;
+
+	void OnCalibratorSelected(const FAssetData& AssetData);
+
+	bool DoesAssetHaveCalibrationComponent(const FAssetData& AssetData) const;
+
+	FAssetData GetCalibratorAssetData() const;
+};
+
+/** Row widget for the calibration row list view */
+class SCalibrationDatasetRow : public SMultiColumnTableRow<TSharedPtr<FCalibrationRow>>
 {
-	SLATE_BEGIN_ARGS(SLensDistortionToolPanel) {}
+public:
+	SLATE_BEGIN_ARGS(SCalibrationDatasetRow) {}
 	SLATE_END_ARGS()
 
-public:
-	void Construct(const FArguments& InArgs, ULensDistortionTool* InTool);
+	void Construct(const FArguments& Args, const TSharedRef<STableViewBase>& OwnerTableView, TSharedPtr<FCalibrationRow>& InRowData);
 
-	/** Refreshes the list of available algorithms shown in the AlgosComboBox */
-	void UpdateAlgosOptions();
+	virtual TSharedRef<SWidget> GenerateWidgetForColumn(const FName& ColumnName) override;
 
 private:
+	TSharedPtr<FCalibrationRow> RowData;
+};
 
-	/** Builds the wrapper for the currently selected algo UI */
-	TSharedRef<SWidget> BuildUIWrapper();
-
-	/** Builds the UI of the solver picker widget */
-	TSharedRef<SWidget> BuildSolverPickerWidget();
+/** UI for the Lens Distortion step */
+class SLensDistortionToolPanel : public SCompoundWidget
+{
+	SLATE_BEGIN_ARGS(SLensDistortionToolPanel) {}
+	SLATE_END_ARGS()
 
-	/** Builds the list of solvers to display when opening the solver picker */
-	TSharedRef<SWidget> BuildSolverList();
+public:
+	void Construct(const FArguments& InArgs, ULensDistortionTool* InTool, TWeakPtr<FCameraCalibrationStepsController> InStepsController);
 
-	/** When the solver class is selected in the UI, update the lens distortion tool with the newly selected class */
-	void OnSolverClassSelected(UClass* SolverClass);
+	/** Refresh the dataset list view */
+	void RefreshListView();
 
-	/** Returns true if the input class is the currently selected one */
-	bool IsSolverClassSelected(UClass* SolverClass) const;
+	/** Hides the modal progress window before the rest of the UI is destroyed */
+	void Shutdown();
 
-	/** Get the display name of the currently selected solver class */
-	FText GetSelectedSolverName() const;
+	/** Shows the modal progress window */
+	void OpenProgressWindow();
 
-	/** Builds the UI for the algorithm picker */
-	TSharedRef<SWidget> BuildAlgoPickerWidget();
+	/** Enables the OkayButton in the progress window */
+	void MarkProgressFinished();
 
-	/** Updates the UI so that it matches the selected algorithm (if necessary) */
-	void UpdateUI();
+	void UpdateProgressText(const FText& ProgressText);
 
 private:
+	/** Callback to handle changes to the customized capture settings */
+	void OnCaptureSettingsChanged(const FPropertyChangedEvent& PropertyChangedEvent);
+
+	TSharedRef<SWidget> BuildDatasetListView();
+	TSharedRef<ITableRow> OnGenerateDatasetRow(TSharedPtr<FCalibrationRow> InItem, const TSharedRef<STableViewBase>& OwnerTable);
+	FReply OnDatasetRowKeyPressed(const FGeometry& Geometry, const FKeyEvent& KeyEvent);
 
-	/** The tool controller object */
-	TWeakObjectPtr<class ULensDistortionTool> Tool;
+	FReply OnClearCalibrationRowsClicked();
+	FReply OnImportDatasetClicked();
+	FReply OnCalibrateClicked();
 
-	/** The box containing the UI given by the selected algorithm */
-	TSharedPtr<class SVerticalBox> UI;
+	void BuildProgressWindow();
+	FReply OnOkPressed();
+	FReply OnCancelPressed();
 	
 private:
+	TWeakObjectPtr<ULensDistortionTool> Tool;
+	TWeakPtr<FCameraCalibrationStepsController> WeakStepsController;
 
-	/** Options source for the AlgosComboBox. Lists the currently available nodal offset algorithms */
-	TArray<TSharedPtr<FString>> CurrentAlgos;
+	TSharedPtr<SListView<TSharedPtr<FCalibrationRow>>> DatasetListView;
+	TSharedPtr<SHeaderRow> DatasetListHeader;
 
-	/** The combobox that presents the available nodal offset algorithms */
-	TSharedPtr<SComboBox<TSharedPtr<FString>>> AlgosComboBox;
+	TSharedPtr<SWindow> ProgressWindow;
+	TSharedPtr<STextBlock> ProgressTextWidget;
+	TSharedPtr<SButton> OkayButton;
 };
diff --git a/Engine/Plugins/VirtualProduction/CameraCalibration/Source/CameraCalibrationEditor/Private/Calibrators/CameraCalibrationSolver.h b/Engine/Plugins/VirtualProduction/CameraCalibration/Source/CameraCalibrationEditor/Private/Calibrators/CameraCalibrationSolver.h
--- a/Engine/Plugins/VirtualProduction/CameraCalibration/Source/CameraCalibrationEditor/Private/Calibrators/CameraCalibrationSolver.h
+++ b/Engine/Plugins/VirtualProduction/CameraCalibration/Source/CameraCalibrationEditor/Private/Calibrators/CameraCalibrationSolver.h
@@ -141,7 +141,7 @@ protected:
   * The implementation is largely based on the implementation of calibrateCamera from OpenCV: [REDACTED_PROFILE_URL]/opencv
   */
 
-UCLASS()
+UCLASS(meta = (DisplayName="OpenCV Solver"))
 class ULensDistortionSolverOpenCV : public ULensDistortionSolver
 {
 	GENERATED_BODY()
diff --git a/Engine/Plugins/VirtualProduction/CameraCalibration/Source/CameraCalibrationEditor/Private/Calibrators/CameraLensDistortionAlgoAruco.cpp b/Engine/Plugins/VirtualProduction/CameraCalibration/Source/CameraCalibrationEditor/Private/Calibrators/CameraLensDistortionAlgoAruco.cpp
deleted file mode 100644
diff --git a/Engine/Plugins/VirtualProduction/CameraCalibration/Source/CameraCalibrationEditor/Private/Calibrators/CameraLensDistortionAlgoAruco.h b/Engine/Plugins/VirtualProduction/CameraCalibration/Source/CameraCalibrationEditor/Private/Calibrators/CameraLensDistortionAlgoAruco.h
deleted file mode 100644
diff --git a/Engine/Plugins/VirtualProduction/CameraCalibration/Source/CameraCalibrationEditor/Private/Calibrators/CameraLensDistortionAlgoCheckerboard.cpp b/Engine/Plugins/VirtualProduction/CameraCalibration/Source/CameraCalibrationEditor/Private/Calibrators/CameraLensDistortionAlgoCheckerboard.cpp
deleted file mode 100644
diff --git a/Engine/Plugins/VirtualProduction/CameraCalibration/Source/CameraCalibrationEditor/Private/Calibrators/CameraLensDistortionAlgoCheckerboard.h b/Engine/Plugins/VirtualProduction/CameraCalibration/Source/CameraCalibrationEditor/Private/Calibrators/CameraLensDistortionAlgoCheckerboard.h
deleted file mode 100644
diff --git a/Engine/Plugins/VirtualProduction/CameraCalibration/Source/CameraCalibrationEditor/Private/Calibrators/CameraLensDistortionAlgoPoints.cpp b/Engine/Plugins/VirtualProduction/CameraCalibration/Source/CameraCalibrationEditor/Private/Calibrators/CameraLensDistortionAlgoPoints.cpp
deleted file mode 100644
diff --git a/Engine/Plugins/VirtualProduction/CameraCalibration/Source/CameraCalibrationEditor/Private/Calibrators/CameraLensDistortionAlgoPoints.h b/Engine/Plugins/VirtualProduction/CameraCalibration/Source/CameraCalibrationEditor/Private/Calibrators/CameraLensDistortionAlgoPoints.h
deleted file mode 100644
diff --git a/Engine/Plugins/VirtualProduction/CameraCalibration/Source/CameraCalibrationEditor/Private/CameraCalibrationUtilsPrivate.cpp b/Engine/Plugins/VirtualProduction/CameraCalibration/Source/CameraCalibrationEditor/Private/CameraCalibrationUtilsPrivate.cpp
--- a/Engine/Plugins/VirtualProduction/CameraCalibration/Source/CameraCalibrationEditor/Private/CameraCalibrationUtilsPrivate.cpp
+++ b/Engine/Plugins/VirtualProduction/CameraCalibration/Source/CameraCalibrationEditor/Private/CameraCalibrationUtilsPrivate.cpp
@@ -130,7 +130,7 @@ namespace UE::CameraCalibration::Private
 		}
 	}
 
-	bool FindArucoCalibrationPoint(const TArray<UCalibrationPointComponent*>& CalibrationComponents, EArucoDictionary ArucoDictionary, const FArucoMarker& ArucoMarker, FArucoCalibrationPoint& OutArucoCalibrationPoint)
+	bool FindArucoCalibrationPoint(const TArray<TWeakObjectPtr<UCalibrationPointComponent>>& CalibrationComponents, EArucoDictionary ArucoDictionary, const FArucoMarker& ArucoMarker, FArucoCalibrationPoint& OutArucoCalibrationPoint)
 	{
 		// Build calibrator point name based on the detected marker
 		const FString DictionaryName = GetArucoDictionaryName(ArucoDictionary);
@@ -141,9 +141,9 @@ namespace UE::CameraCalibration::Private
 		OutArucoCalibrationPoint.MarkerID = ArucoMarker.MarkerID;
 		OutArucoCalibrationPoint.Name = FString::Printf(TEXT("%s-%d"), *DictionaryName, ArucoMarker.MarkerID);
 
-		for (UCalibrationPointComponent* Component : CalibrationComponents)
+		for (TWeakObjectPtr<UCalibrationPointComponent> Component : CalibrationComponents)
 		{
-			if (!Component)
+			if (!Component.IsValid())
 			{
 				continue;
 			}
diff --git a/Engine/Plugins/VirtualProduction/CameraCalibration/Source/CameraCalibrationEditor/Private/CameraCalibrationUtilsPrivate.h b/Engine/Plugins/VirtualProduction/CameraCalibration/Source/CameraCalibrationEditor/Private/CameraCalibrationUtilsPrivate.h
--- a/Engine/Plugins/VirtualProduction/CameraCalibration/Source/CameraCalibrationEditor/Private/CameraCalibrationUtilsPrivate.h
+++ b/Engine/Plugins/VirtualProduction/CameraCalibration/Source/CameraCalibrationEditor/Private/CameraCalibrationUtilsPrivate.h
@@ -2,6 +2,7 @@
 #pragma once
 
 #include "OpenCVHelper.h"
+#include "UObject/WeakObjectPtrTemplates.h"
 
 #include "CameraCalibrationUtilsPrivate.generated.h"
 
@@ -48,7 +49,7 @@ namespace UE::CameraCalibration::Private
 	void FindActorsWithCalibrationComponents(TArray<AActor*>& ActorsWithCalibrationComponents);
 
 	/** Find an aruco marker calibration point in one of the input calibration components that matches the input dictionary and marker ID */
-	bool FindArucoCalibrationPoint(const TArray<UCalibrationPointComponent*>& CalibrationComponents, EArucoDictionary ArucoDictionary, const FArucoMarker& ArucoMarker, FArucoCalibrationPoint& OutArucoCalibrationPoint);
+	bool FindArucoCalibrationPoint(const TArray<TWeakObjectPtr<UCalibrationPointComponent>>& CalibrationComponents, EArucoDictionary ArucoDictionary, const FArucoMarker& ArucoMarker, FArucoCalibrationPoint& OutArucoCalibrationPoint);
 
 	/** Set every pixel in the input texture to the clear color */
 	void ClearTexture(UTexture2D* Texture, FColor ClearColor = FColor::Transparent);
