[Backout] - CL34438771
Original CL Desc
-----------------------------------------------------------------
CameraCalibration: Combine the functionality of the three native lens distortion "algo" classes. The three classes were functionally very similar, and really only differed in how they handled capturing the different calibration patterns. These different capture modes are now all supported by the lens distortion tool itself.
This CL also cleans up some of the UI elements in the lens distortion tab of the lens file editor, replacing most of the custom slate widgets with a details panel (with minor customization).

diff --git a/Engine/Plugins/VirtualProduction/CameraCalibration/Source/CameraCalibrationEditor/Private/AssetEditor/LensDistortionTool.cpp b/Engine/Plugins/VirtualProduction/CameraCalibration/Source/CameraCalibrationEditor/Private/AssetEditor/LensDistortionTool.cpp
--- a/Engine/Plugins/VirtualProduction/CameraCalibration/Source/CameraCalibrationEditor/Private/AssetEditor/LensDistortionTool.cpp
+++ b/Engine/Plugins/VirtualProduction/CameraCalibration/Source/CameraCalibrationEditor/Private/AssetEditor/LensDistortionTool.cpp
@@ -3,99 +3,145 @@
 #include "LensDistortionTool.h"
 
 #include "AssetToolsModule.h"
-#include "CalibrationPointComponent.h"
-#include "Camera/CameraActor.h"
-#include "CameraCalibrationCheckerboard.h"
+#include "Calibrators/CameraCalibrationSolver.h"
 #include "CameraCalibrationEditorLog.h"
 #include "CameraCalibrationSettings.h"
 #include "CameraCalibrationStepsController.h"
-#include "CameraCalibrationUtilsPrivate.h"
+#include "CameraCalibrationTypes.h"
+#include "CameraLensDistortionAlgo.h"
 #include "DesktopPlatformModule.h"
 #include "Dom/JsonObject.h"
 #include "Framework/Application/SlateApplication.h"
 #include "HAL/FileManager.h"
 #include "IDesktopPlatform.h"
 #include "ImageUtils.h"
-#include "JsonObjectConverter.h"
 #include "LensFile.h"
 #include "LensInfoStep.h"
 #include "Misc/DateTime.h"
 #include "Misc/MessageDialog.h"
-#include "Modules/ModuleManager.h"
-#include "OpenCVHelper.h"
+#include "Models/SphericalLensModel.h"
 #include "ScopedTransaction.h"
-#include "Serialization/JsonSerializer.h"
 #include "Serialization/JsonWriter.h"
-#include "SLensDistortionToolPanel.h"
+#include "Serialization/JsonSerializer.h"
+#include "UObject/Class.h"
+#include "UObject/Package.h"
+#include "UObject/UObjectIterator.h"
 
 #define LOCTEXT_NAMESPACE "LensDistortionTool"
 
+#if WITH_EDITOR
+static TAutoConsoleVariable<bool> CVarSolveNodalOffset(TEXT("LensDistortionTool.SolveNodalOffset"), false, TEXT("If true, and if the calibration result contained a valid nodal offset, that offset will be added to the LensFile in addition to distortion."));
+#endif
+
 namespace UE::CameraCalibration::Private::LensDistortionTool
 {
 	static const FString SessionDateTimeField(TEXT("SessionDateTime"));
-	static const FString Version(TEXT("Version"));
+	static const FString AlgoNameField(TEXT("AlgoName"));
 }
 
 void ULensDistortionTool::Initialize(TWeakPtr<FCameraCalibrationStepsController> InCameraCalibrationStepController)
 {
-	WeakStepsController = InCameraCalibrationStepController;
+	CameraCalibrationStepsController = InCameraCalibrationStepController;
+
+	// Find available solver classes
 
-	// Discover all available solver classes and select the first one as the solver to use
 	TArray<UClass*> DerivedSolverClasses;
 	GetDerivedClasses(ULensDistortionSolver::StaticClass(), DerivedSolverClasses);
 
-	if (DerivedSolverClasses.Num() > 0)
-	{
-		SolverSettings.SolverClass = DerivedSolverClasses[0];
-	}
+	check(!DerivedSolverClasses.IsEmpty());
+
+	SetSolverClass(DerivedSolverClasses[0]);
+
+	// Find available algos
 
-	// Find all actors in the current level that have calibration point components and select the first one as the starting calibrator 
-	TArray<AActor*> CalibratorActors;
-	UE::CameraCalibration::Private::FindActorsWithCalibrationComponents(CalibratorActors);
+	TArray<TSubclassOf<UCameraLensDistortionAlgo>> Algos;
 
-	if (CalibratorActors.Num() > 0)
+	for (TObjectIterator<UClass> AlgoIt; AlgoIt; ++AlgoIt)
 	{
-		SetCalibrator(CalibratorActors[0]);
+		if (AlgoIt->IsChildOf(UCameraLensDistortionAlgo::StaticClass()) && !AlgoIt->HasAnyClassFlags(CLASS_Abstract | CLASS_Deprecated))
+		{
+			const UCameraLensDistortionAlgo* Algo = CastChecked<UCameraLensDistortionAlgo>(AlgoIt->GetDefaultObject());
+			AlgosMap.Add(Algo->FriendlyName(), TSubclassOf<UCameraLensDistortionAlgo>(*AlgoIt));
+
+			// If the algo uses an overlay material, create a new MID to use with that algo
+			if (UMaterialInterface* OverlayMaterial = Algo->GetOverlayMaterial())
+			{
+				AlgoOverlayMIDs.Add(Algo->FriendlyName(), UMaterialInstanceDynamic::Create(OverlayMaterial, GetTransientPackage()));
+			}
+		}
 	}
 
-	// Initialize the overlay material and texture
-	if (TSharedPtr<FCameraCalibrationStepsController> StepsController = WeakStepsController.Pin())
+	if (TSharedPtr<FCameraCalibrationStepsController> SharedStepsController = CameraCalibrationStepsController.Pin())
 	{
-		UMaterialInterface* OverlayParent = TSoftObjectPtr<UMaterialInterface>(FSoftObjectPath(TEXT("/CameraCalibration/Materials/M_Coverage.M_Coverage"))).LoadSynchronous();
-		OverlayMID = UMaterialInstanceDynamic::Create(OverlayParent, GetTransientPackage());
+		if (ULensFile* const LensFile = SharedStepsController->GetLensFile())
+		{
+			LensFile->OnLensFileModelChanged().AddUObject(this, &ULensDistortionTool::OnLensModelChanged);
 
-		const FIntPoint CompResolution = StepsController->GetCompRenderResolution();
+			UpdateAlgoMap(LensFile->LensInfo.LensModel);
+		}
+	}
 
-		OverlayTexture = UTexture2D::CreateTransient(CompResolution.X, CompResolution.Y, EPixelFormat::PF_B8G8R8A8);
-		UE::CameraCalibration::Private::ClearTexture(OverlayTexture);
+	BuildProgressWindowWidgets();
+}
 
-		OverlayMID->SetTextureParameterValue(FName(TEXT("CoverageTexture")), OverlayTexture);
+void ULensDistortionTool::OnLensModelChanged(const TSubclassOf<ULensModel>& LensModel)
+{
+	UpdateAlgoMap(LensModel);
+}
+
+void ULensDistortionTool::UpdateAlgoMap(const TSubclassOf<ULensModel>& LensModel)
+{
+	SupportedAlgosMap.Empty();
+	for (const TPair<FName, TSubclassOf<UCameraLensDistortionAlgo>>& AlgoPair : AlgosMap)
+	{
+		const UCameraLensDistortionAlgo* Algo = CastChecked<UCameraLensDistortionAlgo>(AlgoPair.Value->GetDefaultObject());
 
-		StepsController->SetOverlayMaterial(OverlayMID);
+		if (Algo->SupportsModel(LensModel))
+		{
+			SupportedAlgosMap.Add(AlgoPair);
+		}
+	}
+	if (DistortionWidget)
+	{
+		DistortionWidget->UpdateAlgosOptions();
 	}
 }
 
 void ULensDistortionTool::Shutdown()
 {
-	if (CalibrationTask.IsValid())
+	if (CurrentAlgo)
 	{
-		CancelCalibration();
-		CalibrationTask = {};
-		DistortionWidget->Shutdown();
-	}
+		if (CalibrationTask.IsValid())
+		{
+			CurrentAlgo->CancelCalibration();
+			CalibrationTask = {};
+			ProgressWindow->HideWindow();
+		}
 
-	EndCalibrationSession();
-}
+		CurrentAlgo->Shutdown();
+		CurrentAlgo = nullptr;
 
-TSharedRef<SWidget> ULensDistortionTool::BuildUI()
-{
-	DistortionWidget = SNew(SLensDistortionToolPanel, this, WeakStepsController);
-	return DistortionWidget.ToSharedRef();
+		EndCalibrationSession();
+	}
+
+	if (TSharedPtr<FCameraCalibrationStepsController> SharedStepsController = CameraCalibrationStepsController.Pin())
+	{
+		if (ULensFile* const LensFile = SharedStepsController->GetLensFile())
+		{
+			LensFile->OnLensFileModelChanged().RemoveAll(this);
+		}
+	}
 }
 
 void ULensDistortionTool::Tick(float DeltaTime)
 {
+	if (CurrentAlgo)
+	{
+		CurrentAlgo->Tick(DeltaTime);
+	}
+
 	// A valid task handle implies that there is an asynchronous calibration happening on another thread.
+	// The tool will poll the task to determine when it has finished so that the results can be saved.
 	if (CalibrationTask.IsValid())
 	{
 		if (CalibrationTask.IsCompleted())
@@ -104,32 +150,31 @@ void ULensDistortionTool::Tick(float DeltaTime)
 			CalibrationResult = CalibrationTask.GetResult();
 			CalibrationTask = {};
 
-			FText TaskCompletionText;
-			if (CalibrationResult.ErrorMessage.IsEmpty())
+			if (!CalibrationResult.ErrorMessage.IsEmpty())
 			{
-				FNumberFormattingOptions Options;
-				Options.MinimumFractionalDigits = 3;
-				Options.MaximumFractionalDigits = 3;
-
-				TaskCompletionText = FText::Format(LOCTEXT("CalibrationResultReprojectionError", "Reprojection Error: {0} pixels"), FText::AsNumber(CalibrationResult.ReprojectionError, &Options));
+				const FText Message = FText::Format(LOCTEXT("CalibrationErrorResult", "Calibration Error: {0}"), CalibrationResult.ErrorMessage);
+				ProgressTextWidget->SetText(Message);
 			}
 			else
 			{
-				TaskCompletionText = FText::Format(LOCTEXT("CalibrationResultErrorMessage", "Calibration Error: {0}"), CalibrationResult.ErrorMessage);
+				// Update progress window with final reprojection error
+				FFormatOrderedArguments Arguments;
+				Arguments.Add(FText::FromString(FString::Printf(TEXT("%.3f"), CalibrationResult.ReprojectionError)));
+
+				const FText Message = FText::Format(LOCTEXT("CalibrationTaskResult", "Reprojection Error: {0} pixels"), Arguments);
+				ProgressTextWidget->SetText(Message);
 			}
 
-			DistortionWidget->UpdateProgressText(TaskCompletionText);
-			DistortionWidget->MarkProgressFinished();
+			OkayButton->SetEnabled(true);
 		}
 		else
 		{
-			// Update the calibration status in the progress window
 			FText StatusText = FText::GetEmpty();
-			const bool bIsStatusNew = GetCalibrationStatus(StatusText);
+			const bool bIsStatusNew = CurrentAlgo->GetCalibrationStatus(StatusText);
 
 			if (bIsStatusNew)
 			{
-				DistortionWidget->UpdateProgressText(StatusText);
+				ProgressTextWidget->SetText(StatusText);
 			}
 		}
 	}
@@ -137,445 +182,303 @@ void ULensDistortionTool::Tick(float DeltaTime)
 
 bool ULensDistortionTool::OnViewportClicked(const FGeometry& MyGeometry, const FPointerEvent& MouseEvent)
 {
-	// We only respond to left clicks
-	if (MouseEvent.GetEffectingButton() != EKeys::LeftMouseButton)
+	if (!bIsActive)
 	{
 		return false;
 	}
 
-	// Block user interaction with the simulcam viewport while an async calibration task is executing
-	if (CalibrationTask.IsValid())
+	if (!CurrentAlgo)
 	{
 		return false;
 	}
 
-	TSharedPtr<FCameraCalibrationStepsController> StepsController = WeakStepsController.Pin();
-	if (!StepsController)
+	// Block user interaction with the simulcam viewport while an async calibration task is executing
+	if (CalibrationTask.IsValid())
 	{
 		return false;
 	}
 
-	// Create new row of calibration data
-	TSharedPtr<FCalibrationRow> NewRow = MakeShared<FCalibrationRow>();
-
-	// If capturing a single point, add the clicked point to the new row
-	if (CaptureSettings.CalibrationPattern == ECalibrationPattern::Points)
-	{
-		const bool bDetectionResult = DetectPoint(MyGeometry, MouseEvent, NewRow);
-		if (!bDetectionResult)
-		{
-			return true; // Though unsuccessful, the user input was handled
-		}
-	}
-	else
-	{
-		// If capturing a calibration pattern, read the media texture into CPU memory to send to one of the pattern detection algorithms
-		FText OutErrorMessage;
-
-		TArray<FColor> Pixels;
-		FIntPoint ImageSize;
-		if (!StepsController->ReadMediaPixels(Pixels, ImageSize, OutErrorMessage, ESimulcamViewportPortion::CameraFeed))
-		{
-			FMessageDialog::Open(EAppMsgType::Ok, OutErrorMessage, LOCTEXT("ErrorReadingMedia", "Media Error"));
-			return true; // Though unsuccessful, the user input was handled
-		}
-
-		// Detect the selected pattern in the media image
-		bool bDetectionResult = false;
-		if (CaptureSettings.CalibrationPattern == ECalibrationPattern::Checkerboard)
-		{
-			bDetectionResult = DetectCheckerboardPattern(Pixels, ImageSize, NewRow);
-		}
-		else if (CaptureSettings.CalibrationPattern == ECalibrationPattern::Aruco)
-		{
-			bDetectionResult = DetectArucoPattern(Pixels, ImageSize, NewRow);
-		}
-
-		if (!bDetectionResult)
-		{
-			return true; // Though unsuccessful, the user input was handled
-		}
-
-		// Save an image view of the captured frame
-		FImageView ImageView = FImageView(Pixels.GetData(), ImageSize.X, ImageSize.Y, ERawImageFormat::BGRA8);
-		ImageView.CopyTo(NewRow->MediaImage);
-	}
+	return CurrentAlgo->OnViewportClicked(MyGeometry, MouseEvent);
+}
 
-	NewRow->Index = AdvanceSessionRowIndex();
-	NewRow->Pattern = CaptureSettings.CalibrationPattern;
+TSharedRef<SWidget> ULensDistortionTool::BuildUI()
+{
+	DistortionWidget = SNew(SLensDistortionToolPanel, this);
+	return DistortionWidget.ToSharedRef();
+}
 
-	NewRow->CameraPose = FTransform::Identity;
-	if (const ACameraActor* Camera = StepsController->GetCamera())
-	{
-		if (const UCameraComponent* CameraComponent = Camera->GetCameraComponent())
-		{
-			NewRow->CameraPose = CameraComponent->GetComponentToWorld();
-		}
-	}
+void ULensDistortionTool::BuildProgressWindowWidgets()
+{
+	ProgressWindow = SNew(SWindow)
+		.Title(LOCTEXT("ProgressWindowTitle", "Distortion Calibration Progress"))
+		.SizingRule(ESizingRule::Autosized)
+		.IsTopmostWindow(true)
+		.HasCloseButton(false)
+		.SupportsMaximize(false)
+		.SupportsMinimize(true);
+
+	ProgressTextWidget = SNew(STextBlock).Text(FText::GetEmpty());
+
+	OkayButton = SNew(SButton)
+		.IsEnabled(false)
+		.HAlign(HAlign_Center)
+		.VAlign(VAlign_Center)
+		.Text(LOCTEXT("OkText", "Ok"))
+		.OnClicked_UObject(this, &ULensDistortionTool::OnOkPressed);
+
+	TSharedRef<SWidget> WindowContent = SNew(SVerticalBox)
+
+		// Text widget to display the current progress of the calibration
+		+ SVerticalBox::Slot()
+		.HAlign(EHorizontalAlignment::HAlign_Center)
+		.VAlign(EVerticalAlignment::VAlign_Center)
+		[
+			ProgressTextWidget.ToSharedRef()
+		]
+
+		// Ok and Cancel buttons
+		+ SVerticalBox::Slot()
+		.HAlign(EHorizontalAlignment::HAlign_Center)
+		.VAlign(EVerticalAlignment::VAlign_Center)
+		[
+			SNew(SHorizontalBox)
+
+			+ SHorizontalBox::Slot()
+			.AutoWidth()
+			[
+				OkayButton.ToSharedRef()
+			]
+
+			+ SHorizontalBox::Slot()
+			.AutoWidth()
+			[
+				SNew(SButton)
+				.Text(LOCTEXT("CancelText", "Cancel"))
+				.HAlign(HAlign_Center)
+				.VAlign(VAlign_Center)
+				.OnClicked_UObject(this, &ULensDistortionTool::OnCancelPressed)
+			]
+		];
+
+	ProgressWindow->SetContent(WindowContent);
+
+	// Create the window, but start with it hidden. When the user initiates a calibration, the progress window will be shown.
+	FSlateApplication::Get().AddWindow(ProgressWindow.ToSharedRef());
+	ProgressWindow->HideWindow();
+}
 
-	Dataset.CalibrationRows.Add(NewRow);
+bool ULensDistortionTool::DependsOnStep(UCameraCalibrationStep* Step) const
+{
+	return Cast<ULensInfoStep>(Step) != nullptr;
+}
 
-	// Notify the ListView of the new data
-	if (DistortionWidget)
+void ULensDistortionTool::Activate()
+{
+	// Nothing to do if it is already active.
+	if (bIsActive)
 	{
-		DistortionWidget->RefreshListView();
+		return;
 	}
 
-	// Export the data for this row to a .json file on disk
-	ExportCalibrationRow(NewRow);
-
-	ExportSessionData();
-
-	return true;
+	bIsActive = true;
 }
 
-bool ULensDistortionTool::DetectCheckerboardPattern(TArray<FColor>& Pixels, FIntPoint Size, TSharedPtr<FCalibrationRow> OutRow)
+void ULensDistortionTool::Deactivate()
 {
-	const FText ErrorTitle = LOCTEXT("CaptureError", "Capture Error");
-
-	// The selected calibrator must be a checkerboard actor
-	ACameraCalibrationCheckerboard* Checkerboard = Cast<ACameraCalibrationCheckerboard>(CaptureSettings.Calibrator.Get());
-	if (!Checkerboard)
-	{
-		const FText ErrorMessage = LOCTEXT("CheckerboardActorRequiredError", "The selected calibrator must be actor must be a Camera Calibration Checkerboard actor.");
-		FMessageDialog::Open(EAppMsgCategory::Error, EAppMsgType::Ok, ErrorMessage, ErrorTitle);
-		return false;
-	}
+	bIsActive = false;
+}
 
-	OutRow->CheckerboardDimensions = FIntPoint(Checkerboard->NumCornerCols, Checkerboard->NumCornerRows);
+UClass* ULensDistortionTool::GetSolverClass()
+{
+	return SolverClass;
+}
 
-	TArray<FVector2f> DetectedCorners;
-	const bool bCornersFound = FOpenCVHelper::IdentifyCheckerboard(Pixels, Size, OutRow->CheckerboardDimensions, DetectedCorners);
+void ULensDistortionTool::SetSolverClass(UClass* InSolverClass)
+{
+	SolverClass = InSolverClass;
+}
 
-	if (!bCornersFound || DetectedCorners.IsEmpty())
+void ULensDistortionTool::ResetAlgo()
+{
+	// Remove old Algo
+	if (CurrentAlgo)
 	{
-		const FText ErrorMessage = FText::Format(LOCTEXT("NoCheckerboardError", "Failed to detect a {0}x{1} checkerboard in the image."), Checkerboard->NumCornerCols, Checkerboard->NumCornerRows);
-		FMessageDialog::Open(EAppMsgCategory::Error, EAppMsgType::Ok, ErrorMessage, ErrorTitle);
-		return false;
-	}
+		CurrentAlgo->Shutdown();
+		CurrentAlgo = nullptr;
 
-	for (const FVector2f& Corner : DetectedCorners)
-	{
-		OutRow->ImagePoints.Points.Add(FVector2D(Corner));
+		EndCalibrationSession();
 	}
 
-	// Fill out the checkerboard's 3D points
-	const FVector TopLeft = Checkerboard->TopLeft->GetComponentLocation();
-	const FVector TopRight = Checkerboard->TopRight->GetComponentLocation();
-	const FVector BottomLeft = Checkerboard->BottomLeft->GetComponentLocation();
-
-	const FVector RightVector = TopRight - TopLeft;
-	const FVector DownVector = BottomLeft - TopLeft;
-
-	const float HorizontalStep = (Checkerboard->NumCornerCols > 1) ? (1.0f / (Checkerboard->NumCornerCols - 1)) : 0.0f;
-	const float VerticalStep = (Checkerboard->NumCornerRows > 1) ? (1.0f / (Checkerboard->NumCornerRows - 1)) : 0.0f;
-
-	for (int32 RowIdx = 0; RowIdx < Checkerboard->NumCornerRows; ++RowIdx)
+	// Set the tool overlay pass' material to the MID associate with the current algo
+	if (TSharedPtr<FCameraCalibrationStepsController> StepsController = CameraCalibrationStepsController.Pin())
 	{
-		for (int32 ColIdx = 0; ColIdx < Checkerboard->NumCornerCols; ++ColIdx)
-		{
-			const FVector PointLocation = TopLeft + (RightVector * ColIdx * HorizontalStep) + (DownVector * RowIdx * VerticalStep);
-			OutRow->ObjectPoints.Points.Add(PointLocation);
-		}
+		StepsController->SetOverlayEnabled(false);
+		StepsController->SetOverlayMaterial(nullptr);
 	}
-
-	// Update the coverage overlay with the latest checkerboard corners
-	TArray<FVector2D> CameraFeedAdjustedCorners = OutRow->ImagePoints.Points;
-	const FIntPoint OverlayTextureSize = FIntPoint(OverlayTexture->GetSizeX(), OverlayTexture->GetSizeY());
-	RescalePoints(CameraFeedAdjustedCorners, OverlayTextureSize, Size);
-
-	FOpenCVHelper::DrawCheckerboardCorners(DetectedCorners, OutRow->CheckerboardDimensions, OverlayTexture);
-
-	if (TSharedPtr<FCameraCalibrationStepsController> StepsController = WeakStepsController.Pin())
-	{
-		StepsController->RefreshOverlay();
-	}
-
-	OutRow->TargetPose.SetLocation(TopLeft);
-
-	FRotator BoardRotation = FRotationMatrix::MakeFromYZ(TopRight - TopLeft, TopLeft - BottomLeft).Rotator();
-	OutRow->TargetPose.SetRotation(BoardRotation.Quaternion());
-
-	return true;
 }
 
-bool ULensDistortionTool::DetectArucoPattern(TArray<FColor>& Pixels, FIntPoint Size, TSharedPtr<FCalibrationRow> OutRow)
+void ULensDistortionTool::SetAlgo(const FName& AlgoName)
 {
-	const FText ErrorTitle = LOCTEXT("CaptureError", "Capture Error");
+	// Find the algo class
 
-	// Detect the aruco dictionary to use by looking at the names of the calibration points on the selected calibrator actor
-	AActor* CalibratorActor = CaptureSettings.Calibrator.Get();
-	if (!CalibratorActor)
-	{
-		const FText ErrorMessage = LOCTEXT("NoCalibratorError", "Please select a valid calibrator actor.");
-		FMessageDialog::Open(EAppMsgCategory::Error, EAppMsgType::Ok, ErrorMessage, ErrorTitle);
-		return false;
-	}
+	//@todo replace with Find to avoid double search
 
-	EArucoDictionary Dictionary = UE::CameraCalibration::Private::GetArucoDictionaryForCalibrator(CalibratorActor);
-	if (Dictionary == EArucoDictionary::None)
+	if (!AlgosMap.Contains(AlgoName))
 	{
-		const FText ErrorMessage = LOCTEXT("NoArucoDictionaryError", "The calibration components of the selected calibrator do not specify a valid Aruco dictionary.");
-		FMessageDialog::Open(EAppMsgCategory::Error, EAppMsgType::Ok, ErrorMessage, ErrorTitle);
-		return false;
+		return;
 	}
 
-	// Identify any aruco markers matching the current dictionary in the media image
-	TArray<FArucoMarker> IdentifiedMarkers;
-	bool bResult = FOpenCVHelper::IdentifyArucoMarkers(Pixels, Size, Dictionary, IdentifiedMarkers);
-
-	if (!bResult || IdentifiedMarkers.IsEmpty())
+	TSubclassOf<UCameraLensDistortionAlgo>& AlgoClass = AlgosMap[AlgoName];
+		
+	// If it is the same as the existing one, do nothing.
+	if (!CurrentAlgo && !AlgoClass)
 	{
-		const FText ErrorMessage = LOCTEXT("NoArucoMarkersFoundError", "Failed to detect any aruco markers in the image belonging to the dictionary of the selected calibrator.");
-		FMessageDialog::Open(EAppMsgCategory::Error, EAppMsgType::Ok, ErrorMessage, ErrorTitle);
-		return false;
+		return;
 	}
-
-	// For each identified marker, search the calibration components to find the subpoints matching that marker and the 3D location of each of its corners
-	TArray<FArucoCalibrationPoint> ArucoCalibrationPoints;
-	ArucoCalibrationPoints.Reserve(IdentifiedMarkers.Num());
-
-	UpdateCalibrationComponents();
-
-	for (const FArucoMarker& Marker : IdentifiedMarkers)
+	else if (CurrentAlgo && (CurrentAlgo->GetClass() == AlgoClass))
 	{
-		FArucoCalibrationPoint ArucoCalibrationPoint;
-		if (UE::CameraCalibration::Private::FindArucoCalibrationPoint(CalibrationComponents, Dictionary, Marker, ArucoCalibrationPoint))
-		{
-			ArucoCalibrationPoints.Add(ArucoCalibrationPoint);
-		}
+		return;
 	}
 
-	if (ArucoCalibrationPoints.Num() > 0)
+	// Remove old Algo
+	if (CurrentAlgo)
 	{
-		for (const FArucoCalibrationPoint& Marker : ArucoCalibrationPoints)
-		{
-			for (int32 CornerIndex = 0; CornerIndex < 4; ++CornerIndex)
-			{
-				OutRow->ObjectPoints.Points.Add(Marker.Corners3D[CornerIndex]);
-				OutRow->ImagePoints.Points.Add(FVector2D(Marker.Corners2D[CornerIndex]));
-			}
-		}
-
-		FArucoCalibrationPoint& FirstAruco = ArucoCalibrationPoints[0];
-
-		FVector TopLeft = FirstAruco.Corners3D[0];
-		FVector TopRight = FirstAruco.Corners3D[1];
-		FVector BottomLeft = FirstAruco.Corners3D[3];
+		CurrentAlgo->Shutdown();
+		CurrentAlgo = nullptr;
 
-		OutRow->TargetPose.SetLocation(TopLeft);
-
-		FRotator FirstMarkerRotation = FRotationMatrix::MakeFromYZ(TopRight - TopLeft, TopLeft - BottomLeft).Rotator();
-		OutRow->TargetPose.SetRotation(FirstMarkerRotation.Quaternion());
+		EndCalibrationSession();
 	}
 
-	FOpenCVHelper::DrawArucoMarkers(IdentifiedMarkers, OverlayTexture);
-
-	if (TSharedPtr<FCameraCalibrationStepsController> StepsController = WeakStepsController.Pin())
-	{
-		StepsController->RefreshOverlay();
-	}
-
-	return true;
-}
-
-bool ULensDistortionTool::DetectPoint(const FGeometry& MyGeometry, const FPointerEvent& MouseEvent, TSharedPtr<FCalibrationRow> OutRow)
-{
-	TSharedPtr<FCameraCalibrationStepsController> StepsController = WeakStepsController.Pin();
-	if (!StepsController)
+	// If AlgoClass is none, we're done here.
+	if (!AlgoClass)
 	{
-		return false;
+		return;
 	}
 
-	// When the user captures the first calibration point, the 3D locations for the remaining points are cached to ensure they all come from the same frame of data
-	// The media is also paused to ensure that the user is able to more easily capture the 2D location of each subsequent calibration point
-	UpdateCalibrationComponents();
-	
-	if (CalibrationComponentIndex == 0)
-	{
-		StepsController->Pause();
-
-		// Cache the camera pose and component locations of the remaining components
-		CachedComponentLocations.Reset();
-		for (TWeakObjectPtr<UCalibrationPointComponent> Component : CalibrationComponents)
-		{
-			CachedComponentLocations.Add(Component->GetComponentLocation());
-		}
-	}
+	// Create new algo
+	CurrentAlgo = NewObject<UCameraLensDistortionAlgo>(
+		GetTransientPackage(),
+		AlgoClass,
+		MakeUniqueObjectName(GetTransientPackage(), AlgoClass));
 
-	// Calculate the location where the user clicked in the viewport
-	FVector2f NormalizedClickPosition;
-	if (!StepsController->CalculateNormalizedMouseClickPosition(MyGeometry, MouseEvent, NormalizedClickPosition, ESimulcamViewportPortion::CameraFeed))
+	if (CurrentAlgo)
 	{
-		return false;
+		CurrentAlgo->Initialize(this);
 	}
 
-	const FIntPoint ImageSize = StepsController->GetCameraFeedSize();
-	const FVector2D ImagePoint = FVector2D(NormalizedClickPosition * ImageSize);
-
-	OutRow->ImagePoints.Points.Add(ImagePoint);
-
-	OutRow->ObjectPoints.Points.Add(CachedComponentLocations[CalibrationComponentIndex]);
-
-	// Advance the component index, and if it loops around, resume playing the media (which was paused after capturing the first point)
-	SetComponentIndex(CalibrationComponentIndex + 1);
-
-	if (CalibrationComponentIndex == 0)
+	// Set the tool overlay pass' material to the MID associate with the current algo
+	if (CameraCalibrationStepsController.IsValid())
 	{
-		StepsController->Play();
+		TSharedPtr<FCameraCalibrationStepsController> StepsController = CameraCalibrationStepsController.Pin();
+		StepsController->SetOverlayEnabled(false);
+		StepsController->SetOverlayMaterial(GetOverlayMID());
 	}
-
-	return true;
 }
 
-void ULensDistortionTool::CalibrateLens()
+UCameraLensDistortionAlgo* ULensDistortionTool::GetAlgo() const
 {
-	TSharedPtr<FCameraCalibrationStepsController> StepsController = WeakStepsController.Pin();
-	if (!StepsController)
-	{
-		return;
-	}
+	return CurrentAlgo;
+}
 
-	ULensFile* LensFile = StepsController->GetLensFile();
-	if (!LensFile)
+void ULensDistortionTool::OnSaveCurrentCalibrationData()
+{
+	if (!CurrentAlgo)
 	{
 		return;
 	}
 
-	const FText TitleError = LOCTEXT("CalibrationErrorTitle", "Calibration Error");
+	const FText TitleError = LOCTEXT("LensCalibrationError", "Lens Calibration Error");
+	const FText UnknownError = LOCTEXT("UnknownError", "An unknown error occurred initiating the distortion calibration. Check the output log for details.");
 
-	if (Dataset.CalibrationRows.Num() < 1)
+	if (CurrentAlgo->SupportsAsyncCalibration())
 	{
-		FMessageDialog::Open(EAppMsgType::Ok, LOCTEXT("EmptyDatasetError", "The calibration dataset is empty. Please select a valid dataset or capture a new one before calibrating."), TitleError);
-		return;
-	}
+		FText ErrorMessage;
+		CalibrationTask = CurrentAlgo->BeginCalibration(ErrorMessage);
 
-	const float PhysicalSensorWidth = StepsController->GetLensFileEvaluationInputs().Filmback.SensorWidth;
-	const float PixelAspect = LensFile->LensInfo.SqueezeFactor;
+		if (!CalibrationTask.IsValid())
+		{
+			if (!ErrorMessage.IsEmpty())
+			{
+				FMessageDialog::Open(EAppMsgType::Ok, ErrorMessage, TitleError);
+			}
+			else
+			{
+				FMessageDialog::Open(EAppMsgType::Ok, UnknownError, TitleError);
+			}
+			return;
+		}
 
-	const float DesqueezedSensorWidth = PhysicalSensorWidth * PixelAspect;
+		// TODO: This text is temporary. Update to register a delegate with the algo to provide implementation-specific status text
+		ProgressTextWidget->SetText(LOCTEXT("CalibrationProgressText", "Calibrating Lens Distortion..."));
 
-	if (FMath::IsNearlyZero(DesqueezedSensorWidth))
-	{
-		FMessageDialog::Open(EAppMsgType::Ok, LOCTEXT("InvalidSensorWidthError", "The sensor width and squeeze factor in the camera settings must both be greater than zero. Please enter a valid value."), TitleError);
-		return;
-	}
+		// Ensure that the Ok button is disabled and show the progress window
+		OkayButton->SetEnabled(false);
+		ProgressWindow->ShowWindow();
 
-	if (!SolverSettings.FocalLengthGuess.IsSet() || FMath::IsNearlyZero(SolverSettings.FocalLengthGuess.GetValue()))
-	{
-		FMessageDialog::Open(EAppMsgType::Ok, LOCTEXT("InvalidFocalLengthError", "Please enter a valid estimate for the focal length of the lens (in mm)."), TitleError);
-		return;
+		DistortionWidget->SetEnabled(false);
 	}
-
-	const FIntPoint ImageSize = LensFile->CameraFeedInfo.GetDimensions();
-
-	const double FocalLengthEstimateValue = SolverSettings.FocalLengthGuess.GetValue();
-	const double Fx = (FocalLengthEstimateValue / DesqueezedSensorWidth) * ImageSize.X;
-
-	// When operating on a desqueezed image, we expect our pixel aspect to be square, so horizontal and vertical field of view are assumed to be equal (i.e. Fx == Fy)
-	FVector2D FocalLength = FVector2D(Fx);
-	FVector2D ImageCenter = FVector2D((ImageSize.X - 1) * 0.5, (ImageSize.Y - 1) * 0.5);
-
-	TArray<FObjectPoints> Samples3d;
-	Samples3d.Reserve(Dataset.CalibrationRows.Num());
-
-	TArray<FImagePoints> Samples2d;
-	Samples2d.Reserve(Dataset.CalibrationRows.Num());
-
-	TArray<FTransform> CameraPoses;
-	CameraPoses.Reserve(Dataset.CalibrationRows.Num());
-
-	TArray<FTransform> TargetPoses;
-	TargetPoses.Reserve(Dataset.CalibrationRows.Num());
-
-	// Extract the 3D points, 2D points, and camera poses from each row to pass to the solver
-	for (const TSharedPtr<FCalibrationRow>& Row : Dataset.CalibrationRows)
+	else
 	{
-		Samples3d.Add(Row->ObjectPoints);
-		Samples2d.Add(Row->ImagePoints);
-		CameraPoses.Add(Row->CameraPose);
-		TargetPoses.Add(Row->TargetPose);
-	}
+		TSubclassOf<ULensModel> LensModel;
+		bool bResult = CurrentAlgo->GetLensDistortion(
+			CalibrationResult.EvaluatedFocus, 
+			CalibrationResult.EvaluatedZoom, 
+			CalibrationResult.Parameters, 
+			CalibrationResult.FocalLength, 
+			CalibrationResult.ImageCenter, 
+			LensModel, 
+			CalibrationResult.ReprojectionError, 
+			CalibrationResult.ErrorMessage);
+
+		if (!bResult)
+		{
+			if (CalibrationResult.ErrorMessage.IsEmpty())
+			{
+				CalibrationResult.ErrorMessage = UnknownError;
+			}
 
-	ECalibrationFlags SolverFlags = ECalibrationFlags::None;
-	EnumAddFlags(SolverFlags, ECalibrationFlags::UseIntrinsicGuess);
+			FMessageDialog::Open(EAppMsgType::Ok, CalibrationResult.ErrorMessage, TitleError);
+			return;
+		}
 
-	if (CaptureSettings.bIsCameraTracked)
-	{
-		EnumAddFlags(SolverFlags, ECalibrationFlags::UseExtrinsicGuess);
-	}
+		// Update progress window with final reprojection error
+		ProgressWindow->ShowWindow();
 
-	if (CaptureSettings.bIsCalibratorTracked && CaptureSettings.CalibrationPattern != ECalibrationPattern::Points)
-	{
-		EnumAddFlags(SolverFlags, ECalibrationFlags::SolveTargetOffset);
-	}
+		FFormatOrderedArguments Arguments;
+		Arguments.Add(FText::FromString(FString::Printf(TEXT("%.3f"), CalibrationResult.ReprojectionError)));
 
-	if (SolverSettings.bFixFocalLength)
-	{
-		EnumAddFlags(SolverFlags, ECalibrationFlags::FixFocalLength);
-	}
+		const FText Message = FText::Format(LOCTEXT("ReprojectionError", "Reprojection Error: {0} pixels"), Arguments);
+		ProgressTextWidget->SetText(Message);
 
-	if (SolverSettings.bFixImageCenter)
-	{
-		EnumAddFlags(SolverFlags, ECalibrationFlags::FixPrincipalPoint);
+		OkayButton->SetEnabled(true);
 	}
+}
 
-	if (SolverSettings.bFixDistortion)
-	{
-		EnumAddFlags(SolverFlags, ECalibrationFlags::FixDistortion);
-	}
+FReply ULensDistortionTool::OnCancelPressed()
+{
+	CurrentAlgo->CancelCalibration();
 
-	const FLensFileEvaluationInputs LensFileEvalInputs = StepsController->GetLensFileEvaluationInputs();
+	CalibrationTask = {};
+	ProgressWindow->HideWindow();
+	DistortionWidget->SetEnabled(true);
 
-	FDistortionInfo DistortionGuess;
-	LensFile->EvaluateDistortionParameters(LensFileEvalInputs.Focus, LensFileEvalInputs.Zoom, DistortionGuess);
+	return FReply::Handled();
+}
 
-	const TSubclassOf<ULensModel> Model = LensFile->LensInfo.LensModel;
+FReply ULensDistortionTool::OnOkPressed()
+{
+	SaveCalibrationResult();
 
-	Solver = NewObject<ULensDistortionSolver>(GetTransientPackage(), SolverSettings.SolverClass);
+	ProgressWindow->HideWindow();
+	DistortionWidget->SetEnabled(true);
 
-	CalibrationTask = UE::Tasks::Launch(UE_SOURCE_LOCATION, [InSolver = this->Solver, Model, Samples3d, Samples2d, ImageSize, FocalLength, ImageCenter, DistortionGuess, CameraPoses, TargetPoses, PixelAspect, SolverFlags, LensFileEvalInputs]() mutable
-		{
-			FDistortionCalibrationResult Result = InSolver->Solve(
-				Samples3d,
-				Samples2d,
-				ImageSize,
-				FocalLength,
-				ImageCenter,
-				DistortionGuess.Parameters,
-				CameraPoses,
-				TargetPoses,
-				Model,
-				PixelAspect,
-				SolverFlags
-			);
-
-			// CalibrateCamera() returns focal length and image center in pixels, but the result is expected to be normalized by the image size
-			Result.FocalLength.FxFy = Result.FocalLength.FxFy / ImageSize;
-			Result.ImageCenter.PrincipalPoint = Result.ImageCenter.PrincipalPoint / ImageSize;
-
-			// FZ inputs to LUT
-			Result.EvaluatedFocus = LensFileEvalInputs.Focus;
-			Result.EvaluatedZoom = LensFileEvalInputs.Zoom;
-
-			return Result;
-		});
-
-	DistortionWidget->OpenProgressWindow();
-
-	// All of the UI options should be disabled while the calibration task is running
-	DistortionWidget->SetEnabled(false);
+	return FReply::Handled();
 }
 
 void ULensDistortionTool::SaveCalibrationResult()
 {
-	TSharedPtr<FCameraCalibrationStepsController> StepsController = WeakStepsController.Pin();
-	if (!StepsController)
-	{
-		return;
-	}
+	ULensFile* LensFile = CameraCalibrationStepsController.Pin()->GetLensFile();
 
-	ULensFile* LensFile = StepsController->GetLensFile();
 	if (!LensFile)
 	{
 		return;
@@ -593,7 +496,7 @@ void ULensDistortionTool::SaveCalibrationResult()
 		CalibrationResult.STMap.DistortionMap = (ImportedImages.Num() > 0) ? Cast<UTexture>(ImportedImages[0]) : nullptr;
 	}
 
-	// The result may have featured calibrated distortion parameters or an ST Map.
+	// Depending on the algo, it is possible that the result feature calibrated distortion parameters or an ST Map.
 	// If the result contains any distortion parameters, then the results will be written as a distortion point in the Lens File
 	// Otherwise, if the result contains a valid ST Map, then it will be added to the Lens File
 	if (CalibrationResult.Parameters.Parameters.Num() > 0)
@@ -627,7 +530,7 @@ void ULensDistortionTool::SaveCalibrationResult()
 	}
 
 	// If the calibration result contains a meaningful nodal offset result, add it to the Lens File
-	if (SolverSettings.CalibrationType == ECalibrationType::NodalOffset || SolverSettings.CalibrationType == ECalibrationType::Everything)
+	if (CVarSolveNodalOffset.GetValueOnGameThread())
 	{
 		if (!CalibrationResult.NodalOffset.LocationOffset.Equals(FVector::ZeroVector) || !CalibrationResult.NodalOffset.RotationOffset.Equals(FQuat::Identity))
 		{
@@ -638,186 +541,52 @@ void ULensDistortionTool::SaveCalibrationResult()
 		}
 	}
 
-	ClearCalibrationRows();
-}
-
-void ULensDistortionTool::CancelCalibration()
-{
-	if (Solver)
-	{
-		Solver->Cancel();
-	}
-}
-
-bool ULensDistortionTool::GetCalibrationStatus(FText& StatusText) const
-{
-	if (Solver)
-	{
-		return Solver->GetStatusText(StatusText);
-	}
-	return false;
-}
-
-void ULensDistortionTool::SetCalibrator(AActor* InCalibrator)
-{
-	CaptureSettings.Calibrator = InCalibrator;
-
-	if (!InCalibrator)
+	if (UCameraLensDistortionAlgo* Algo = GetAlgo())
 	{
-		return;
+		Algo->OnDistortionSavedToLens();
 	}
-
-	CalibrationComponents.Reset();
-	UpdateCalibrationComponents();
-
-	// Initialize the component index used by the single point detection mode
-	SetComponentIndex(0);
-
-	CachedComponentLocations.Reserve(CalibrationComponents.Num());
 }
 
-void ULensDistortionTool::UpdateCalibrationComponents()
+FCameraCalibrationStepsController* ULensDistortionTool::GetCameraCalibrationStepsController() const
 {
-	bool bNeedsUpdate = false;
-
-	for (TWeakObjectPtr<UCalibrationPointComponent> Component : CalibrationComponents)
+	if (!CameraCalibrationStepsController.IsValid())
 	{
-		if (!Component.IsValid())
-		{
-			bNeedsUpdate = true;
-			break;
-		}
+		return nullptr;
 	}
 
-	if (CalibrationComponents.IsEmpty())
-	{
-		bNeedsUpdate = true;
-	}
-
-	if (bNeedsUpdate && CaptureSettings.Calibrator.IsValid())
-	{
-		// Find all of the calibration components attached to calibrator actor
-		TArray<UCalibrationPointComponent*> CalibrationPoints;
-		CaptureSettings.Calibrator->GetComponents(CalibrationPoints);
-
-		// Store weak references to all of the calibration components that have an attached scene component
-		CalibrationComponents.Reset();
-		for (UCalibrationPointComponent* CalibrationPoint : CalibrationPoints)
-		{
-			if (CalibrationPoint && CalibrationPoint->GetAttachParent())
-			{
-				CalibrationComponents.Add(CalibrationPoint);
-			}
-		}
-	}
+	return CameraCalibrationStepsController.Pin().Get();
 }
 
-void ULensDistortionTool::ClearCalibrationRows()
+TArray<FName> ULensDistortionTool::GetAlgos() const
 {
-	Dataset.CalibrationRows.Empty();
-
-	if (DistortionWidget)
-	{
-		DistortionWidget->RefreshListView();
-	}
-
-	// Reset the calibration component index to restart the pattern
-	SetComponentIndex(0);
-
-	RefreshCoverage();
-
-	// End the current calibration session (a new one will begin the next time a new row is added)
-	EndCalibrationSession();
+	TArray<FName> OutKeys;
+	SupportedAlgosMap.GetKeys(OutKeys);
+	return OutKeys;
 }
 
-void ULensDistortionTool::SetComponentIndex(int32 Index)
+bool ULensDistortionTool::IsActive() const
 {
-	CalibrationComponentIndex = Index;
-
-	// If the Index would loop around, reset it to 0
-	if (CalibrationComponentIndex == CalibrationComponents.Num())
-	{
-		CalibrationComponentIndex = 0;
-	}
-
-	UpdateCalibrationComponents();
-
-	if (CalibrationComponents.IsValidIndex(CalibrationComponentIndex))
-	{
-		CaptureSettings.NextPoint = FText::FromString(CalibrationComponents[CalibrationComponentIndex]->GetName());
-	}
+	return bIsActive;
 }
 
-void ULensDistortionTool::RefreshCoverage()
+UMaterialInstanceDynamic* ULensDistortionTool::GetOverlayMID() const
 {
-	UE::CameraCalibration::Private::ClearTexture(OverlayTexture);
-
-	for (const TSharedPtr<FCalibrationRow>& Row : Dataset.CalibrationRows)
-	{
-		if (Row->Pattern == ECalibrationPattern::Checkerboard)
-		{
-			TArray<FVector2D> CameraFeedAdjustedCorners = Row->ImagePoints.Points;
-			const FIntPoint OverlayTextureSize = FIntPoint(OverlayTexture->GetSizeX(), OverlayTexture->GetSizeY());
-			const FIntPoint ImageSize = FIntPoint(Row->MediaImage.SizeX, Row->MediaImage.SizeY);
-			RescalePoints(CameraFeedAdjustedCorners, OverlayTextureSize, ImageSize);
-
-			FOpenCVHelper::DrawCheckerboardCorners(CameraFeedAdjustedCorners, Row->CheckerboardDimensions, OverlayTexture);
-		}
-	}
-
-	// The coverage texture may have changed as a result of a change in size or pixel format.
-	// Therefore, the material parameter should be updated to ensure it is up to date.
-	if (OverlayTexture)
+	if (CurrentAlgo)
 	{
-		OverlayMID->SetTextureParameterValue(FName(TEXT("CoverageTexture")), OverlayTexture);
+		return AlgoOverlayMIDs.FindRef(CurrentAlgo->FriendlyName()).Get();
 	}
 
-	if (TSharedPtr<FCameraCalibrationStepsController> StepsController = WeakStepsController.Pin())
-	{
-		StepsController->RefreshOverlay();
-	}
+	return nullptr;
 }
 
-void ULensDistortionTool::RescalePoints(TArray<FVector2D>& Points, FIntPoint DebugTextureSize, FIntPoint CameraFeedSize)
+bool ULensDistortionTool::IsOverlayEnabled() const
 {
-	// It is possible that the size of the debug texture is different than the size of the camera feed.
-	// Therefore, the input points should be shifted so that they appear at the correct location in the debug image.
-	const FVector2D TopLeftCorner = (DebugTextureSize - CameraFeedSize) / 2.0f;
-
-	for (FVector2D& Point : Points)
+	if (CurrentAlgo)
 	{
-		Point += TopLeftCorner;
+		return CurrentAlgo->IsOverlayEnabled();
 	}
-}
-
-bool ULensDistortionTool::DependsOnStep(UCameraCalibrationStep* Step) const
-{
-	return Cast<ULensInfoStep>(Step) != nullptr;
-}
 
-void ULensDistortionTool::Activate()
-{
-	bIsActive = true;
-}
-
-void ULensDistortionTool::Deactivate()
-{
-	bIsActive = false;
-}
-
-bool ULensDistortionTool::IsActive() const
-{
-	return bIsActive;
-}
-
-UMaterialInstanceDynamic* ULensDistortionTool::GetOverlayMID() const
-{
-	return OverlayMID;
-}
-
-bool ULensDistortionTool::IsOverlayEnabled() const
-{
-	return CaptureSettings.bShowOverlay;
+	return false;
 }
 
 void ULensDistortionTool::StartCalibrationSession()
@@ -850,7 +619,8 @@ FString ULensDistortionTool::GetSessionSaveDir() const
 
 	const FString SessionDateString = SessionInfo.StartTime.ToString(TEXT("%Y-%m-%d"));
 	const FString SessionTimeString = SessionInfo.StartTime.ToString(TEXT("%H-%M-%S"));
-	const FString DatasetDir = SessionTimeString;
+	const FString DatasetPrefix = TEXT("Dataset-") + CurrentAlgo->ShortName().ToString() + TEXT("Algorithm-");
+	const FString DatasetDir = DatasetPrefix + SessionTimeString;
 
 	const FString ProjectSaveDir = FPaths::ProjectSavedDir() / TEXT("CameraCalibration") / TEXT("LensDistortion");
 
@@ -889,6 +659,84 @@ void ULensDistortionTool::DeleteExportedRow(const int32& RowIndex) const
 	}
 }
 
+void ULensDistortionTool::ExportCalibrationRow(int32 RowIndex, const TSharedRef<FJsonObject>& RowObject, const FImageView& RowImage)
+{
+	if (!GetDefault<UCameraCalibrationSettings>()->IsCalibrationDatasetImportExportEnabled())
+	{
+		return;
+	}
+
+	// Start a calibration session (if one is not currently active)
+	StartCalibrationSession();
+
+	// Assemble the path and filename for this row based on the session and row index
+	const FString PathName = GetSessionSaveDir();
+	const FString FileName = GetRowFilename(RowIndex) + FDateTime::Now().ToString(TEXT("%H-%M-%S"));
+
+	const FString JsonFileName = PathName / FileName + TEXT(".json");
+	const FString ImageFileName = PathName / FileName + TEXT(".png");
+
+	// Create and open a new Json file for writing, and initialize a JsonWriter to serialize the contents
+	if (TUniquePtr<FArchive> FileWriter = TUniquePtr<FArchive>(IFileManager::Get().CreateFileWriter(*JsonFileName)))
+	{
+		TSharedRef< TJsonWriter<> > JsonWriter = TJsonWriterFactory<>::Create(FileWriter.Get());
+
+		// Write the Json row data out and save the file
+		FJsonSerializer::Serialize(RowObject, JsonWriter);
+		FileWriter->Close();
+
+		UE_LOG(LogCameraCalibrationEditor, Verbose, TEXT("Lens Distortion Tool wrote to dataset row file: %s"), *JsonFileName);
+	}
+
+	// If the row has an image to export, save it out to a file
+	if (RowImage.RawData != nullptr)
+	{
+		FImageUtils::SaveImageByExtension(*ImageFileName, RowImage);
+	}
+}
+
+void ULensDistortionTool::ExportSessionData(const TSharedRef<FJsonObject>& SessionDataObject)
+{
+	using namespace UE::CameraCalibration::Private;
+
+	if (!GetDefault<UCameraCalibrationSettings>()->IsCalibrationDatasetImportExportEnabled())
+	{
+		return;
+	}
+
+	// Start a calibration session (if one is not currently active)
+	StartCalibrationSession();
+
+	// Assemble the path and filename for this row based on the session and row index
+	const FString PathName = GetSessionSaveDir();
+	const FString FileName = TEXT("SessionData");
+
+	const FString SessionFileName = PathName / FileName + TEXT(".ucamcalib");
+
+	// Delete the existing session data file (if it exists)
+	if (IFileManager::Get().FileExists(*SessionFileName))
+	{
+		IFileManager::Get().Delete(*SessionFileName);
+	}
+
+	// Create and open a new Json file for writing, and initialize a JsonWriter to serialize the contents
+	if (TUniquePtr<FArchive> FileWriter = TUniquePtr<FArchive>(IFileManager::Get().CreateFileWriter(*SessionFileName)))
+	{
+		TSharedRef< TJsonWriter<> > JsonWriter = TJsonWriterFactory<>::Create(FileWriter.Get());
+
+		const FString SessionDateTimeString = SessionInfo.StartTime.ToString(TEXT("%Y-%m-%d")) + TEXT("-") + SessionInfo.StartTime.ToString(TEXT("%H-%M-%S"));
+
+		SessionDataObject->SetStringField(LensDistortionTool::SessionDateTimeField, SessionDateTimeString);
+		SessionDataObject->SetStringField(LensDistortionTool::AlgoNameField, CurrentAlgo->FriendlyName().ToString());
+
+		// Write the Json row data out and save the file
+		FJsonSerializer::Serialize(SessionDataObject, JsonWriter);
+		FileWriter->Close();
+
+		UE_LOG(LogCameraCalibrationEditor, Verbose, TEXT("Lens Distortion Tool wrote to dataset session file: %s"), *SessionFileName);
+	}
+}
+
 void ULensDistortionTool::ImportCalibrationDataset()
 {
 	using namespace UE::CameraCalibration::Private;
@@ -899,7 +747,7 @@ void ULensDistortionTool::ImportCalibrationDataset()
 	}
 
 	// If there is existing calibration data that will be overwritten during import, ask the user to confirm that they want to continue
-	if (Dataset.CalibrationRows.Num() > 0)
+	if (CurrentAlgo->HasCalibrationData())
 	{
 		const FText ConfirmationMessage = LOCTEXT(
 			"ImportDatasetConfirmationMessage",
@@ -971,6 +819,30 @@ void ULensDistortionTool::ImportCalibrationDataset()
 				{
 					UE_LOG(LogCameraCalibrationEditor, Verbose, TEXT("Lens Distortion Tool failed to deserialize the date/time from the session file: %s"), *SessionFile);
 				}
+
+				// Import the algo name so we can set the appropriate algo before importing the row data
+				FString AlgoString;
+				if (JsonSessionData->TryGetStringField(LensDistortionTool::AlgoNameField, AlgoString))
+				{
+					// Ensure that the algo name matches one of the algos for this tool
+					const FName AlgoName = FName(*AlgoString);
+					if (AlgosMap.Contains(AlgoName))
+					{
+						SetAlgo(AlgoName);
+					}
+					else
+					{
+						const FText ErrorMessage = LOCTEXT("UnknownCalibrationAlgo", "The selected dataset does not represent a lens distortion calibration. Choose a different .ucamcalib dataset file from a lens distortion calibration.");
+						FMessageDialog::Open(EAppMsgType::Ok, ErrorMessage);
+						return;
+					}
+				}
+				else
+				{
+					UE_LOG(LogCameraCalibrationEditor, Verbose, TEXT("Lens Distortion Tool failed to deserialize the algo name from the session file: %s"), *SessionFile);
+				}
+
+				CurrentAlgo->ImportSessionData(JsonSessionData.ToSharedRef());
 			}
 			else
 			{
@@ -979,7 +851,7 @@ void ULensDistortionTool::ImportCalibrationDataset()
 		}
 	}
 
-	Dataset.CalibrationRows.Empty();
+	CurrentAlgo->PreImportCalibrationData();
 
 	// Initialize a maximum row index which will be used to set the current session row index if the user wants to add additional rows after importing
 	int32 MaxRowIndex = -1;
@@ -1005,7 +877,7 @@ void ULensDistortionTool::ImportCalibrationDataset()
 			TSharedPtr<FJsonObject> JsonRowData = MakeShared<FJsonObject>();
 			if (FJsonSerializer::Deserialize(JsonReader, JsonRowData))
 			{
-				int32 RowNum = ImportCalibrationRow(JsonRowData.ToSharedRef(), RowImage);
+				int32 RowNum = CurrentAlgo->ImportCalibrationRow(JsonRowData.ToSharedRef(), RowImage);
 				MaxRowIndex = FMath::Max(MaxRowIndex, RowNum);
 			}
 			else
@@ -1015,14 +887,7 @@ void ULensDistortionTool::ImportCalibrationDataset()
 		}
 	}
 
-	// Sort imported calibration rows by row index
-	Dataset.CalibrationRows.Sort([](const TSharedPtr<FCalibrationRow>& LHS, const TSharedPtr<FCalibrationRow>& RHS) { return LHS->Index < RHS->Index; });
-
-	// Notify the ListView of the new data
-	if (DistortionWidget)
-	{
-		DistortionWidget->RefreshListView();
-	}
+	CurrentAlgo->PostImportCalibrationData();
 
 	// Set the current session's start date/time and row index to match what was just imported to support adding/deleting rows
 	SessionInfo.bIsActive = true;
@@ -1030,115 +895,4 @@ void ULensDistortionTool::ImportCalibrationDataset()
 	SessionInfo.RowIndex = MaxRowIndex;
 }
 
-void ULensDistortionTool::ExportSessionData()
-{
-	if (!GetDefault<UCameraCalibrationSettings>()->IsCalibrationDatasetImportExportEnabled())
-	{
-		return;
-	}
-
-	using namespace UE::CameraCalibration::Private;
-
-	TSharedPtr<FJsonObject> JsonSessionData = MakeShared<FJsonObject>();
-	JsonSessionData->SetNumberField(LensDistortionTool::Version, 2);
-
-	// Start a calibration session (if one is not currently active)
-	StartCalibrationSession();
-
-	// Assemble the path and filename for this row based on the session and row index
-	const FString PathName = GetSessionSaveDir();
-	const FString FileName = TEXT("SessionData");
-
-	const FString SessionFileName = PathName / FileName + TEXT(".ucamcalib");
-
-	// Delete the existing session data file (if it exists)
-	if (IFileManager::Get().FileExists(*SessionFileName))
-	{
-		IFileManager::Get().Delete(*SessionFileName);
-	}
-
-	// Create and open a new Json file for writing, and initialize a JsonWriter to serialize the contents
-	if (TUniquePtr<FArchive> FileWriter = TUniquePtr<FArchive>(IFileManager::Get().CreateFileWriter(*SessionFileName)))
-	{
-		TSharedRef< TJsonWriter<> > JsonWriter = TJsonWriterFactory<>::Create(FileWriter.Get());
-
-		const FString SessionDateTimeString = SessionInfo.StartTime.ToString(TEXT("%Y-%m-%d")) + TEXT("-") + SessionInfo.StartTime.ToString(TEXT("%H-%M-%S"));
-
-		JsonSessionData->SetStringField(LensDistortionTool::SessionDateTimeField, SessionDateTimeString);
-
-		// Write the Json row data out and save the file
-		FJsonSerializer::Serialize(JsonSessionData.ToSharedRef(), JsonWriter);
-		FileWriter->Close();
-
-		UE_LOG(LogCameraCalibrationEditor, Verbose, TEXT("Lens Distortion Tool wrote to dataset session file: %s"), *SessionFileName);
-	}
-}
-
-void ULensDistortionTool::ExportCalibrationRow(TSharedPtr<FCalibrationRow> Row)
-{
-	if (!GetDefault<UCameraCalibrationSettings>()->IsCalibrationDatasetImportExportEnabled())
-	{
-		return;
-	}
-
-	if (const TSharedPtr<FJsonObject>& RowObject = FJsonObjectConverter::UStructToJsonObject<FCalibrationRow>(Row.ToSharedRef().Get()))
-	{
-		// Start a calibration session (if one is not currently active)
-		StartCalibrationSession();
-
-		// Assemble the path and filename for this row based on the session and row index
-		const FString PathName = GetSessionSaveDir();
-		const FString FileName = GetRowFilename(Row->Index) + FDateTime::Now().ToString(TEXT("%H-%M-%S"));
-
-		const FString JsonFileName = PathName / FileName + TEXT(".json");
-		const FString ImageFileName = PathName / FileName + TEXT(".png");
-
-		// Create and open a new Json file for writing, and initialize a JsonWriter to serialize the contents
-		if (TUniquePtr<FArchive> FileWriter = TUniquePtr<FArchive>(IFileManager::Get().CreateFileWriter(*JsonFileName)))
-		{
-			TSharedRef< TJsonWriter<> > JsonWriter = TJsonWriterFactory<>::Create(FileWriter.Get());
-
-			// Write the Json row data out and save the file
-			FJsonSerializer::Serialize(RowObject.ToSharedRef(), JsonWriter);
-			FileWriter->Close();
-
-			UE_LOG(LogCameraCalibrationEditor, Verbose, TEXT("Lens Distortion Tool wrote to dataset row file: %s"), *JsonFileName);
-		}
-
-		// If the row has an image to export, save it out to a file
-		FImageView ImageView = Row->MediaImage;
-		if (ImageView.RawData != nullptr)
-		{
-			FImageUtils::SaveImageByExtension(*ImageFileName, ImageView);
-		}
-	}
-}
-
-int32 ULensDistortionTool::ImportCalibrationRow(const TSharedRef<FJsonObject>& CalibrationRowObject, const FImage& RowImage)
-{
-	// Create a new row to populate with data from the Json object
-	TSharedPtr<FCalibrationRow> NewRow = MakeShared<FCalibrationRow>();
-
-	// We enforce strict mode to ensure that every field in the UStruct of row data is present in the imported json.
-	// If any fields are missing, it is likely the row will be invalid, which will lead to errors in the calibration.
-	constexpr int64 CheckFlags = 0;
-	constexpr int64 SkipFlags = 0;
-	constexpr bool bStrictMode = true;
-	if (FJsonObjectConverter::JsonObjectToUStruct<FCalibrationRow>(CalibrationRowObject, NewRow.Get(), CheckFlags, SkipFlags, bStrictMode))
-	{
-		if (!RowImage.RawData.IsEmpty())
-		{
-			NewRow->MediaImage = RowImage;
-		}
-
-		Dataset.CalibrationRows.Add(NewRow);
-	}
-	else
-	{
-		UE_LOG(LogCameraCalibrationEditor, Warning, TEXT("Failed to import calibration row because at least one field could not be deserialized from the json file."));
-	}
-
-	return NewRow->Index;
-}
-
 #undef LOCTEXT_NAMESPACE
diff --git a/Engine/Plugins/VirtualProduction/CameraCalibration/Source/CameraCalibrationEditor/Private/AssetEditor/LensDistortionTool.h b/Engine/Plugins/VirtualProduction/CameraCalibration/Source/CameraCalibrationEditor/Private/AssetEditor/LensDistortionTool.h
--- a/Engine/Plugins/VirtualProduction/CameraCalibration/Source/CameraCalibrationEditor/Private/AssetEditor/LensDistortionTool.h
+++ b/Engine/Plugins/VirtualProduction/CameraCalibration/Source/CameraCalibrationEditor/Private/AssetEditor/LensDistortionTool.h
@@ -2,12 +2,15 @@
 
 #pragma once
 
-#include "CameraCalibrationStep.h"
+#include "CoreMinimal.h"
 
-#include "Calibrators/CameraCalibrationSolver.h"
+#include "CameraCalibrationStep.h"
 #include "CameraCalibrationTypes.h"
-#include "Engine/Texture2D.h"
+#include "CameraLensDistortionAlgo.h"
+#include "Engine/World.h"
 #include "ImageCore.h"
+#include "SLensDistortionToolPanel.h"
+#include "Widgets/Input/SButton.h"
 
 #include "LensDistortionTool.generated.h"
 
@@ -15,136 +18,8 @@ struct FGeometry;
 struct FPointerEvent;
 
 class FJsonObject;
-class UCalibrationPointComponent;
-
-typedef UE::Tasks::TTask<FDistortionCalibrationResult> FDistortionCalibrationTask;
-
-/** List of supported calibration patterns */
-UENUM()
-enum class ECalibrationPattern : uint8
-{
-	Checkerboard,
-	Aruco,
-	Points
-};
-
-/** List of supported calibration workflows */
-UENUM()
-enum class ECalibrationType : uint8
-{
-	DistortionAndIntrinsics,
-	NodalOffset,
-	Everything
-};
-
-/** The data associated with a single captured calibration pattern/point */
-USTRUCT()
-struct FCalibrationRow
-{
-	GENERATED_BODY()
-
-	/** Index to display in list view */
-	UPROPERTY()
-	int32 Index = -1;
-
-	/** Set of captured 3D points for the calibrator in world space */
-	UPROPERTY()
-	FObjectPoints ObjectPoints;
-
-	/** Set of captured 2D pixel locations where the calibrator was detected in the image */
-	UPROPERTY()
-	FImagePoints ImagePoints;
-
-	/** Pose of the camera actor when this data for this row was captured */
-	UPROPERTY()
-	FTransform CameraPose = FTransform::Identity;
-
-	/** Pose of the calibrator actor when this data for this row was captured */
-	UPROPERTY()
-	FTransform TargetPose = FTransform::Identity;
-
-	/** The calibration pattern used to capture the data for this row */
-	UPROPERTY()
-	ECalibrationPattern Pattern = ECalibrationPattern::Checkerboard;
-
-	/** Dimensions of the detected checkerboard pattern (only valid if calibration pattern was ECalibrationPattern::Checkerboard) */
-	UPROPERTY()
-	FIntPoint CheckerboardDimensions = FIntPoint(0, 0);
-
-	/** Stored frame from the media source associated with the data for this row */
-	FImage MediaImage;
-};
-
-/** An array of captured rows with calibration data */
-USTRUCT()
-struct FCalibrationDataset
-{
-	GENERATED_BODY()
-
-	TArray<TSharedPtr<FCalibrationRow>> CalibrationRows;
-};
-
-/** Settings that control how data is captured in the tool */
-USTRUCT()
-struct FLensCaptureSettings
-{
-	GENERATED_BODY()
-
-	/** The pattern to detect in the media image */
-	UPROPERTY(EditAnywhere, Category = "Capture Settings")
-	ECalibrationPattern CalibrationPattern = ECalibrationPattern::Checkerboard;
-
-	/** An actor with Calibration Point Components that represents the virtual version of a real calibration target */
-	UPROPERTY(EditAnywhere, Category = "Capture Settings")
-	TWeakObjectPtr<AActor> Calibrator;
-
-	/** Set to true if the calibration target being used is tracked */
-	UPROPERTY(EditAnywhere, Category = "Capture Settings")
-	bool bIsCalibratorTracked = false;
-
-	/** Set to true if the camera being used is tracked */
-	UPROPERTY(EditAnywhere, Category = "Capture Settings")
-	bool bIsCameraTracked = false;
-
-	/** Display a debug overlay over the simulcam viewport showing the detected patterns that have been captured */
-	UPROPERTY(EditAnywhere, Category = "Capture Settings", meta = (EditCondition = "CalibrationPattern != ECalibrationPattern::Points", EditConditionHides))
-	bool bShowOverlay = true;
-
-	/** The name of the next calibration point to locate in the image */
-	UPROPERTY(VisibleAnywhere, Category = "Capture Settings", meta = (EditCondition = "CalibrationPattern == ECalibrationPattern::Points", EditConditionHides))
-	FText NextPoint;
-};
-
-/** Settings that control what and how the solver will compute the calibrated lens data */
-USTRUCT()
-struct FLensSolverSettings
-{
-	GENERATED_BODY()
-
-	/** Select the solver to use when solving for calibrated lens data */
-	UPROPERTY(EditAnywhere, NoClear, Category = "Solver Settings", meta = (HideViewOptions, ShowDisplayNames))
-	TSubclassOf<ULensDistortionSolver> SolverClass = ULensDistortionSolverOpenCV::StaticClass();
-
-	/** Select the data that the solver should compute and write to the Lens File */
-	UPROPERTY(EditAnywhere, Category = "Solver Settings")
-	ECalibrationType CalibrationType = ECalibrationType::DistortionAndIntrinsics;
-
-	/** An estimate for the focal length of the lens */
-	UPROPERTY(EditAnywhere, Category = "Solver Settings")
-	TOptional<double> FocalLengthGuess;
-
-	/** Set to true to prevent the solver from optimizing the focal length. The focal length guess will be used. */
-	UPROPERTY(EditAnywhere, AdvancedDisplay, Category = "Solver Settings", meta = (DisplayName = "Fix Focal Length During Optimization"))
-	bool bFixFocalLength = false;
-
-	/** Set to true to prevent the solver from optimizing the image center. The image center will be assumed to be in the exact center. */
-	UPROPERTY(EditAnywhere, AdvancedDisplay, Category = "Solver Settings", meta = (DisplayName = "Fix Image Center During Optimization"))
-	bool bFixImageCenter = false;
-
-	/** Set to true to prevent the solver from optimizing the distortion parameters. The current calibrated distortion parameters will be used. */
-	UPROPERTY(EditAnywhere, AdvancedDisplay, Category = "Solver Settings", meta = (DisplayName = "Fix Distortion During Optimization"))
-	bool bFixDistortion = false;
-};
+class UCameraLensDistortionAlgo;
+class ULensModel;
 
 /** Data associated with a lens distortion calibration session */
 struct FLensDistortionSessionInfo
@@ -160,7 +35,7 @@ struct FLensDistortionSessionInfo
 };
 
 /**
- * The primary tool used in the LensFile asset editor to capture and solve for calibrated lens data
+ * ULensDistortionTool is the controller for the lens distortion panel.
  */
 UCLASS()
 class ULensDistortionTool : public UCameraCalibrationStep
@@ -168,8 +43,7 @@ class ULensDistortionTool : public UCameraCalibrationStep
 	GENERATED_BODY()
 
 public:
-	friend class SLensDistortionToolPanel;
-	
+
 	//~ Begin UCameraCalibrationStep interface
 	virtual void Initialize(TWeakPtr<FCameraCalibrationStepsController> InCameraCalibrationStepController) override;
 	virtual void Shutdown() override;
@@ -180,125 +54,125 @@ public:
 	virtual bool DependsOnStep(UCameraCalibrationStep* Step) const override;
 	virtual void Activate() override;
 	virtual void Deactivate() override;
+	virtual FCameraCalibrationStepsController* GetCameraCalibrationStepsController() const override;
 	virtual bool IsActive() const override;
 	virtual UMaterialInstanceDynamic* GetOverlayMID() const override;
 	virtual bool IsOverlayEnabled() const override;
 	//~ End UCameraCalibrationStep interface
 
-private:
+	/** Get the currently selected lens distortion solver class */
+	UClass* GetSolverClass();
+
+	/** Set the lens distortion solver class to use */
+	void SetSolverClass(UClass* InSolverClass);
+
+	/** Selects the algorithm by name */
+	void SetAlgo(const FName& AlgoName);
+
+	/** Resets the current algo to none */
+	void ResetAlgo();
+
+	/** Returns the currently selected algorithm */
+	UCameraLensDistortionAlgo* GetAlgo() const;
+
+	/** Returns available algorithm names */
+	TArray<FName> GetAlgos() const;
+
+	/** Called by the UI when the user wants to save the calibration data that the current algorithm is providing */
+	void OnSaveCurrentCalibrationData();
+
+	/** Triggered when the Lens Model used by the LensFile changes, allowing the tool to update its list of supported algos */
+	void OnLensModelChanged(const TSubclassOf<ULensModel>& LensModel);
+
 	/** Initiate a new calibration session (if one is not already active) */
 	void StartCalibrationSession();
 
 	/** End the current calibration session */
 	void EndCalibrationSession();
 
-	/** Deletes the .json file with the input row index that was previously exported for this session. */
-	void DeleteExportedRow(const int32& RowIndex) const;
-
 	/** Increments the session index and returns its new value */
 	uint32 AdvanceSessionRowIndex();
 
-	/** Called by the UI when the user wants to save the calibration data */
-	void CalibrateLens();
+	/** Intended to be called by one of the algo objects. Exports the data needed by the algo for the current session to a .json file on disk. */
+	void ExportSessionData(const TSharedRef<FJsonObject>& SessionDataObject);
 
-	/** Save the results from a finished distortion calibration to the Lens File */
-	void SaveCalibrationResult();
+	/** Intended to be called by one of the algo objects. Exports the data associated with a single calibration row to a .json file on disk. */
+	void ExportCalibrationRow(int32 RowIndex, const TSharedRef<FJsonObject>& RowObject, const FImageView& RowImage = FImageView());
 
-	/** Cancel an in-progress async calibration task */
-	void CancelCalibration();
+	/** Intended to be called by one of the algo objects. Deletes the .json file with the input row index that was previously exported for this session. */
+	void DeleteExportedRow(const int32& RowIndex) const;
 
-	/** Get the latest status from the lens distortion calibration. Returns true if the status has changed. */
-	bool GetCalibrationStatus(FText& StatusText) const;
+	/** Import a saved calibration data set from disk */
+	void ImportCalibrationDataset();
 
-	/** Detects a checkerboard pattern in the input array of pixels (supplied by the media image) and fills out the 2D and 3D points in the output calibration row */
-	bool DetectCheckerboardPattern(TArray<FColor>& Pixels, FIntPoint Size, TSharedPtr<FCalibrationRow> OutRow);
+private:
+	/** Get the directory for the current session */
+	FString GetSessionSaveDir() const;
 
-	/** Detects an aruco pattern in the input array of pixels (supplied by the media image) and fills out the 2D and 3D points in the output calibration row */
-	bool DetectArucoPattern(TArray<FColor>& Pixels, FIntPoint Size, TSharedPtr<FCalibrationRow> OutRow);
+	/** Get the filename of a row with the input index */
+	FString GetRowFilename(int32 RowIndex) const;
 
-	/** Captures the mouse position and the 3D world location of the current calibration component and and fills out the output calibration row */
-	bool DetectPoint(const FGeometry& MyGeometry, const FPointerEvent& MouseEvent, TSharedPtr<FCalibrationRow> OutRow);
+	/** Update the list of algos that support the current Lens Model */
+	void UpdateAlgoMap(const TSubclassOf<ULensModel>& LensModel);
 
-	/** Import a saved calibration data set from disk */
-	void ImportCalibrationDataset();
+	/** Build the UI for the progress window, including the window content and buttons */
+	void BuildProgressWindowWidgets();
 
-	/** Sets the calibrator object to be used. Updates the selection picker. */
-	void SetCalibrator(AActor* InCalibrator);
+	/** Called in response to user clicking "Cancel" button in the progress window to cancel a currently running calibration task */
+	FReply OnCancelPressed();
 
-	/** Caches the calibration point components of the current calibrator actor, and updates the list if those components become invalid */
-	void UpdateCalibrationComponents();
+	/** Called in response to user clicking "Ok" button in the progress window to save the result from a finished calibration task */
+	FReply OnOkPressed();
 
-	/** Clears the list of calibration rows */
-	void ClearCalibrationRows();
+	/** Save the results from a finished distortion calibration to the Lens File */
+	void SaveCalibrationResult();
 
-	/** Updates the calibration point component index (used when the calbration pattern is ECalibrationPattern::Points) */
-	void SetComponentIndex(int32 Index);
+public:
+	/** Stores info about the current calibration session of this tool */
+	FLensDistortionSessionInfo SessionInfo;
 
-	/** Clears and redraws the overlay texture with the data from the captured calibration rows */
-	void RefreshCoverage();
+private:
 
-	/** Adjust the 2D location of every input pixel to account for differences between the size of the captured image and the overlay texture */
-	void RescalePoints(TArray<FVector2D>& Points, FIntPoint DebugTextureSize, FIntPoint CameraFeedSize);
+	/** Pointer to the calibration steps controller */
+	TWeakPtr<FCameraCalibrationStepsController> CameraCalibrationStepsController;
 
-	/** Get the directory for the current session */
-	FString GetSessionSaveDir() const;
+	/** The currently selected algorithm */
+	UPROPERTY(Transient)
+	TObjectPtr<UCameraLensDistortionAlgo> CurrentAlgo;
 
-	/** Get the filename of a row with the input index */
-	FString GetRowFilename(int32 RowIndex) const;
+	/** The solver class that will be used by the algos when creating a solver to calibrate for distortion */
+	TObjectPtr<UClass> SolverClass;
 
-	/**
-	 * Import a JsonObject of calibration data that represents a single calibration row.
-	 * Returns the row index of the imported row.
-	 */
-	int32 ImportCalibrationRow(const TSharedRef<FJsonObject>& CalibrationRowObject, const FImage& RowImage);
+	/** Holds the registered camera lens distortion algos */
+	UPROPERTY(Transient)
+	TMap<FName, TSubclassOf<UCameraLensDistortionAlgo>> AlgosMap;
 
-	/** Export global session data to a .json file */
-	void ExportSessionData();
+	/** Holds a subset of the registered algos that support the current Lens Model */
+	UPROPERTY(Transient)
+	TMap<FName, TSubclassOf<UCameraLensDistortionAlgo>> SupportedAlgosMap;
 
-	/** Export the row data to a json file */
-	void ExportCalibrationRow(TSharedPtr<FCalibrationRow> Row);
+	/** Map of algo names to overlay MIDs used by those algos */
+	UPROPERTY(Transient)
+	TMap<FName, TObjectPtr<UMaterialInstanceDynamic>> AlgoOverlayMIDs;
 
-private:
-	/** UI Widget for this tool */
+	/** True if this tool is the active one in the panel */
+	bool bIsActive = false;
+
+	/** UI Widget for this Tool */
 	TSharedPtr<SLensDistortionToolPanel> DistortionWidget;
 
-	FLensSolverSettings SolverSettings;
-	FLensCaptureSettings CaptureSettings;
+	/** UI Widget for the calibration task progress window */
+	TSharedPtr<SWindow> ProgressWindow;
 
-	/** Collection of calibration rows containing the data needed to run a lens calibration */
-	FCalibrationDataset Dataset;
+	/** UI Widget for the text to be displayed in the progress window */
+	TSharedPtr<STextBlock> ProgressTextWidget;
 
-	/** The solver object that runs the lens calibration on another thread. The ptr is stored to access the solver to get status updates and/or cancel the calibration from the game thread */
-	TObjectPtr<ULensDistortionSolver> Solver;
+	/** UI Widget for the "Ok" button in the progress window */
+	TSharedPtr<SButton> OkayButton;
 
 	/** An asynchronous task handle. When valid, the tool will poll its state to determine when the task has completed, and then extract the calibration result from this task handle. */
 	FDistortionCalibrationTask CalibrationTask;
 
 	/** The result from the most recently completed distortion calibration. It is saved as a member variable because the result is not saved immediately, but only if the user confirms that it should be saved. */
 	FDistortionCalibrationResult CalibrationResult;
-
-	/** Container for the set of calibrator components selected in the component combobox */
-	TArray<TWeakObjectPtr<UCalibrationPointComponent>> CalibrationComponents;
-
-	/** 
-	 * The 3D world locations of the calibration point components should all come from the same frame.
-	 * Therefore, the locations are all saved at once when the first calibration point is selected (only applicable to ECalibrationPattern::Points)
-	 */
-	TArray<FVector> CachedComponentLocations;
-
-	/** Index of the next calibration point component to capture (only applicable to ECalibrationPattern::Points) */
-	int32 CalibrationComponentIndex = -1;
-
-	/** Capture session info, used to track whether newly captured data should be added to an existing dataset */
-	FLensDistortionSessionInfo SessionInfo;
-
-	/** True if this tool is the active tab in the UI */
-	bool bIsActive = false;
-
-	/** Material and Texture that draw the coverage overlay on top of the simulcam viewport during capture */
-	TObjectPtr<UMaterialInstanceDynamic> OverlayMID;
-	TObjectPtr<UTexture2D> OverlayTexture;
-
-	/** Weak ptr to the steps controller that created this step */
-	TWeakPtr<FCameraCalibrationStepsController> WeakStepsController;
 };
diff --git a/Engine/Plugins/VirtualProduction/CameraCalibration/Source/CameraCalibrationEditor/Private/AssetEditor/SLensDistortionToolPanel.cpp b/Engine/Plugins/VirtualProduction/CameraCalibration/Source/CameraCalibrationEditor/Private/AssetEditor/SLensDistortionToolPanel.cpp
--- a/Engine/Plugins/VirtualProduction/CameraCalibration/Source/CameraCalibrationEditor/Private/AssetEditor/SLensDistortionToolPanel.cpp
+++ b/Engine/Plugins/VirtualProduction/CameraCalibration/Source/CameraCalibrationEditor/Private/AssetEditor/SLensDistortionToolPanel.cpp
@@ -2,507 +2,380 @@
 
 #include "SLensDistortionToolPanel.h"
 
-#include "CalibrationPointComponent.h"
-#include "DetailLayoutBuilder.h"
-#include "Framework/Application/SlateApplication.h"
-#include "ImageUtils.h"
-#include "IStructureDetailsView.h"
-#include "Modules/ModuleManager.h"
+#include "AssetRegistry/AssetData.h"
+#include "Calibrators/CameraCalibrationSolver.h"
+#include "CameraCalibrationSettings.h"
+#include "CameraCalibrationSubsystem.h"
+#include "CameraLensDistortionAlgo.h"
+#include "Dialog/SCustomDialog.h"
+#include "EditorFontGlyphs.h"
+#include "Engine/Selection.h"
+#include "Framework/MultiBox/MultiBoxBuilder.h"
+#include "LensDistortionTool.h"
+#include "LensFile.h"
+#include "Styling/AppStyle.h"
+#include "Styling/CoreStyle.h"
+#include "Styling/SlateStyle.h"
 #include "UI/CameraCalibrationWidgetHelpers.h"
-#include "UI/SFilterableActorPicker.h"
-#include "UI/SImageTexture.h"
+#include "Widgets/Images/SImage.h"
 #include "Widgets/Input/SButton.h"
+#include "Widgets/Input/SComboBox.h"
 #include "Widgets/SBoxPanel.h"
-#include "Widgets/SWindow.h"
+#include "Widgets/Text/STextBlock.h"
+#include "Widgets/Views/STreeView.h"
 
-#define LOCTEXT_NAMESPACE "SLensDistortionToolPanel"
+#define LOCTEXT_NAMESPACE "LensDistortionTool"
 
-void SCalibrationDatasetRow::Construct(const FArguments& Args, const TSharedRef<STableViewBase>& OwnerTableView, TSharedPtr<FCalibrationRow>& InRowData)
-{
-	RowData = InRowData;
-
-	FSuperRowType::FArguments StyleArguments = FSuperRowType::FArguments()
-		.Padding(1.0f)
-		.Style(&FAppStyle().GetWidgetStyle<FTableRowStyle>("TableView.AlternatingRow"));
-
-	SMultiColumnTableRow<TSharedPtr<FCalibrationRow>>::Construct(StyleArguments, OwnerTableView);
-}
-
-TSharedRef<SWidget> SCalibrationDatasetRow::GenerateWidgetForColumn(const FName& ColumnName)
-{
-	if (ColumnName == TEXT("Index"))
-	{
-		return SNew(SVerticalBox)
-			+ SVerticalBox::Slot()
-			.HAlign(EHorizontalAlignment::HAlign_Center)
-			.VAlign(EVerticalAlignment::VAlign_Center)
-			[SNew(STextBlock).Text(FText::AsNumber(RowData->Index))];
-	}
-	else if (ColumnName == TEXT("Image"))
-	{
-		// Generate a transient thumbnail texture to display in the tool
-		const FImage& MediaImage = RowData->MediaImage;
-
-		if ((MediaImage.SizeX < 1) || (MediaImage.SizeY < 1))
-		{
-			const FString Text = FString::Printf(TEXT("Image Unavailable"));
-			return SNew(STextBlock).Text(FText::FromString(Text));
-		}
-
-		constexpr int32 ResolutionDivider = 4;
-		const FIntPoint ThumbnailSize = FIntPoint(MediaImage.SizeX / ResolutionDivider, MediaImage.SizeY / ResolutionDivider);
-
-		FImage ThumbnailImage;
-		FImageCore::ResizeTo(MediaImage, ThumbnailImage, ThumbnailSize.X, ThumbnailSize.Y, MediaImage.Format, MediaImage.GetGammaSpace());
-
-		if (UTexture2D* Thumbnail = FImageUtils::CreateTexture2DFromImage(ThumbnailImage))
-		{
-			return SNew(SImageTexture, Thumbnail)
-				.MinDesiredHeight(4 * FCameraCalibrationWidgetHelpers::DefaultRowHeight)
-				.MaxDesiredHeight(4 * FCameraCalibrationWidgetHelpers::DefaultRowHeight);
-		}
-		else
-		{
-			const FString Text = FString::Printf(TEXT("Image Unavailable"));
-			return SNew(STextBlock).Text(FText::FromString(Text));
-		}
-	}
-	else if (ColumnName == TEXT("ImagePoint"))
-	{
-		const FString Text = FString::Printf(TEXT("(%.2f, %.2f)"),
-			RowData->ImagePoints.Points[0].X,
-			RowData->ImagePoints.Points[0].Y);
-
-		return SNew(STextBlock).Text(FText::FromString(Text));
-	}
-	else if (ColumnName == TEXT("ObjectPoint"))
-	{
-		const FString Text = FString::Printf(TEXT("(%.2f, %.2f, %.2f)"),
-			RowData->ObjectPoints.Points[0].X,
-			RowData->ObjectPoints.Points[0].Y,
-			RowData->ObjectPoints.Points[0].Z);
-
-		return SNew(STextBlock).Text(FText::FromString(Text));
-	}
-
-	return SNullWidget::NullWidget;
-}
-
-void FCaptureSettingsCustomization::CustomizeDetails(IDetailLayoutBuilder& DetailBuilder)
-{
-	PropertyHandle = DetailBuilder.GetProperty(GET_MEMBER_NAME_CHECKED(FLensCaptureSettings, Calibrator));
-
-	IDetailPropertyRow* CalibratorRow = DetailBuilder.EditDefaultProperty(PropertyHandle);
-	CalibratorRow->CustomWidget()
-		.NameContent()
-		[
-			PropertyHandle->CreatePropertyNameWidget()
-		]
-		.ValueContent()
-		[
-			SNew(SFilterableActorPicker)
-				.OnSetObject(this, &FCaptureSettingsCustomization::OnCalibratorSelected)
-				.OnShouldFilterAsset(this, &FCaptureSettingsCustomization::DoesAssetHaveCalibrationComponent)
-				.ActorAssetData(this, &FCaptureSettingsCustomization::GetCalibratorAssetData)
-		];
-
-	NextPointPropertyHandle = DetailBuilder.GetProperty(GET_MEMBER_NAME_CHECKED(FLensCaptureSettings, NextPoint));
-	IDetailPropertyRow* NextPointRow = DetailBuilder.EditDefaultProperty(NextPointPropertyHandle);
-	NextPointRow->CustomWidget()
-		.NameContent()
-		[
-			NextPointPropertyHandle->CreatePropertyNameWidget()
-		]
-		.ValueContent()
-		[
-			SNew(STextBlock)
-				.Text(this, &FCaptureSettingsCustomization::GetNextPointName)
-				.Font(DetailBuilder.GetDetailFont())
-		];
-}
-
-FText FCaptureSettingsCustomization::GetNextPointName() const
-{
-	FText Name;
-	NextPointPropertyHandle->GetValue(Name);
-	return Name;
-}
 
-void FCaptureSettingsCustomization::OnCalibratorSelected(const FAssetData& AssetData)
-{
-	if (AssetData.IsValid())
-	{
-		PropertyHandle->SetValue(AssetData);
-	}
-}
-
-bool FCaptureSettingsCustomization::DoesAssetHaveCalibrationComponent(const FAssetData& AssetData) const
-{
-	if (const AActor* Actor = Cast<AActor>(AssetData.GetAsset()))
-	{
-		constexpr uint32 NumInlineAllocations = 32;
-		TArray<UCalibrationPointComponent*, TInlineAllocator<NumInlineAllocations>> CalibrationPoints;
-		Actor->GetComponents(CalibrationPoints);
-
-		return (CalibrationPoints.Num() > 0);
-	}
-
-	return false;
-}
-
-FAssetData FCaptureSettingsCustomization::GetCalibratorAssetData() const
-{
-	FAssetData Tmp;
-	PropertyHandle->GetValue(Tmp);
-	return Tmp;
-}
-
-void SLensDistortionToolPanel::Construct(const FArguments& InArgs, ULensDistortionTool* InTool, TWeakPtr<FCameraCalibrationStepsController> InStepsController)
+void SLensDistortionToolPanel::Construct(const FArguments& InArgs, ULensDistortionTool* InTool)
 {
 	Tool = InTool;
-	WeakStepsController = InStepsController;
-
-	BuildProgressWindow();
-
-	FPropertyEditorModule& PropertyEditor = FModuleManager::Get().LoadModuleChecked<FPropertyEditorModule>(TEXT("PropertyEditor"));
-
-	FStructureDetailsViewArgs DefaultStructuDetailsViewArgs;
-	FDetailsViewArgs DefaultDetailsViewArgs;
-	DefaultDetailsViewArgs.bAllowSearch = false;
 
-	TSharedRef<FStructOnScope> SolverSettingsStruct = MakeShared<FStructOnScope>(FLensSolverSettings::StaticStruct(), reinterpret_cast<uint8*>(&Tool->SolverSettings));
-	TSharedPtr<IStructureDetailsView> SolverSettingsDetailsView = PropertyEditor.CreateStructureDetailView(DefaultDetailsViewArgs, DefaultStructuDetailsViewArgs, SolverSettingsStruct);
-
-	TSharedRef<FStructOnScope> CaptureSettingsStruct = MakeShared<FStructOnScope>(FLensCaptureSettings::StaticStruct(), reinterpret_cast<uint8*>(&Tool->CaptureSettings));
-
-	TSharedPtr<IStructureDetailsView> CaptureSettingsDetailsView = PropertyEditor.CreateStructureDetailView(DefaultDetailsViewArgs, DefaultStructuDetailsViewArgs, TSharedPtr<FStructOnScope>());
-	CaptureSettingsDetailsView->GetDetailsView()->OnFinishedChangingProperties().AddSP(this, &SLensDistortionToolPanel::OnCaptureSettingsChanged);
-
-	CaptureSettingsDetailsView->GetDetailsView()->RegisterInstancedCustomPropertyLayout(FLensCaptureSettings::StaticStruct(),
-		FOnGetDetailCustomizationInstance::CreateLambda([]() { return MakeShared<FCaptureSettingsCustomization>(); }));
-
-	CaptureSettingsDetailsView->SetStructureData(CaptureSettingsStruct);
+	// This will be the widget wrapper of the custom algo UI.
+	UI = SNew(SVerticalBox);
 
 	ChildSlot
 	[
 		SNew(SHorizontalBox)
 
-		+ SHorizontalBox::Slot()
+		+ SHorizontalBox::Slot() // Right toolbar
 		.FillWidth(0.25f)
-		[
+		[ 
 			SNew(SVerticalBox)
 
-			+ SVerticalBox::Slot()
-			.AutoHeight()
-			[
-				CaptureSettingsDetailsView->GetWidget().ToSharedRef()
-			]
+			+ SVerticalBox::Slot() // Solver picker
+			.MaxHeight(FCameraCalibrationWidgetHelpers::DefaultRowHeight)
+			[FCameraCalibrationWidgetHelpers::BuildLabelWidgetPair(LOCTEXT("SolverPickerWidgetText", "Lens Distortion Solver"), BuildSolverPickerWidget())]
 
-			+ SVerticalBox::Slot()
-			.AutoHeight()
-			[
-				SolverSettingsDetailsView->GetWidget().ToSharedRef()
-			]
-
-			+ SVerticalBox::Slot()
-			.AutoHeight()
-			.MaxHeight(12 * FCameraCalibrationWidgetHelpers::DefaultRowHeight)
-			[
-				BuildDatasetListView()
-			]
+			+ SVerticalBox::Slot() // Algo picker
+			.MaxHeight(FCameraCalibrationWidgetHelpers::DefaultRowHeight)
+			[ FCameraCalibrationWidgetHelpers::BuildLabelWidgetPair(LOCTEXT("LensDistortionAlgo", "Lens Distortion Algo"), BuildAlgoPickerWidget())]
 
-			+ SVerticalBox::Slot()
+			+ SVerticalBox::Slot() // Algo UI
 			.AutoHeight()
-			.HAlign(HAlign_Center)
-			.Padding(0, 20)
-			[
-				SNew(SButton)
-				.Text(LOCTEXT("ClearAll", "Clear All"))
-				.HAlign(HAlign_Center)
-				.VAlign(VAlign_Center)
-				.OnClicked(this, &SLensDistortionToolPanel::OnClearCalibrationRowsClicked)
-			]
+			[ BuildUIWrapper() ]
 
-			+ SVerticalBox::Slot()
+			+ SVerticalBox::Slot() // Import dataset
 			.AutoHeight()
 			.HAlign(HAlign_Center)
 			.Padding(0, 20)
 			[
-				SNew(SButton)
-				.Text(LOCTEXT("ImportDataset", "Import Dataset"))
-				.HAlign(HAlign_Center)
-				.VAlign(VAlign_Center)
-				.OnClicked(this, &SLensDistortionToolPanel::OnImportDatasetClicked)
+				SNew(SHorizontalBox)
+
+				+ SHorizontalBox::Slot()
+				.AutoWidth()
+				[
+					SNew(SButton).Text(LOCTEXT("ImportCalibrationDataset", "Import Lens Distortion Dataset"))
+					.HAlign(HAlign_Center)
+					.VAlign(VAlign_Center)
+					.OnClicked_Lambda([this]() -> FReply
+					{
+						if (ULensDistortionTool* LensDistortionTool = Tool.Get())
+						{
+							LensDistortionTool->ImportCalibrationDataset();
+
+							// After importing, the tool may have switched the active algo, so redraw the UI accordingly
+							UCameraLensDistortionAlgo* Algo = LensDistortionTool->GetAlgo();
+							UI->ClearChildren();
+							UI->AddSlot()[Algo->BuildUI()];
+
+							for (const TSharedPtr<FString>& AlgoString : CurrentAlgos)
+							{
+								if (AlgoString->Equals(Algo->FriendlyName().ToString()))
+								{
+									AlgosComboBox->SetSelectedItem(AlgoString);
+									break;
+								}
+							}
+						}
+						return FReply::Handled();
+					})
+					.Visibility_Lambda([]() -> EVisibility
+					{
+						if (GetDefault<UCameraCalibrationSettings>()->IsCalibrationDatasetImportExportEnabled())
+						{
+							return EVisibility::Visible;
+						}
+						return EVisibility::Collapsed;
+					})
+				]
 			]
 
-			+ SVerticalBox::Slot()
+			+ SVerticalBox::Slot() // Save Offset
 			.AutoHeight()
-			.HAlign(HAlign_Center)
 			.Padding(0, 20)
 			[
-				SNew(SButton)
-				.Text(LOCTEXT("CalibrateLens", "Calibrate Lens"))
+				SNew(SButton).Text(LOCTEXT("AddToLUT", "Add To Lens Distortion Calibration"))
 				.HAlign(HAlign_Center)
 				.VAlign(VAlign_Center)
-				.OnClicked(this, &SLensDistortionToolPanel::OnCalibrateClicked)
+				.OnClicked_Lambda([WeakTool = Tool]() -> FReply
+				{
+					if (WeakTool.IsValid())
+					{
+						WeakTool->OnSaveCurrentCalibrationData();
+					}
+					return FReply::Handled();
+				})
 			]
 		]
 	];
 }
 
-void SLensDistortionToolPanel::OpenProgressWindow()
+TSharedRef<SWidget> SLensDistortionToolPanel::BuildUIWrapper()
 {
-	ProgressTextWidget->SetText(LOCTEXT("CalibrationProgressText", "Calibrating..."));
-
-	// The okay button will be disabled until the calibration is complete. 
-	OkayButton->SetEnabled(false);
-	ProgressWindow->ShowWindow();
-}
+	return SNew(SVerticalBox)
 
-void SLensDistortionToolPanel::MarkProgressFinished()
-{
-	OkayButton->SetEnabled(true);
+		+ SVerticalBox::Slot() // Algo's Widget
+		.AutoHeight()
+		[ UI.ToSharedRef() ];
 }
 
-void SLensDistortionToolPanel::UpdateProgressText(const FText& ProgressText)
+void SLensDistortionToolPanel::UpdateUI()
 {
-	ProgressTextWidget->SetText(ProgressText);
-}
+	check(AlgosComboBox.IsValid());
 
-void SLensDistortionToolPanel::Shutdown()
-{
-	ProgressWindow->HideWindow();
-}
+	// Get current algo to later compare with new one
+	const UCameraLensDistortionAlgo* OldAlgo = Tool->GetAlgo();
 
-void SLensDistortionToolPanel::OnCaptureSettingsChanged(const FPropertyChangedEvent& PropertyChangedEvent)
-{
-	if (!Tool.IsValid())
+	// Set new algo by name
+	TSharedPtr<FString> SelectedAlgo = AlgosComboBox->GetSelectedItem();
+	if (!SelectedAlgo)
 	{
-		return;
+		Tool->ResetAlgo();
 	}
-
-	const FName PropertyName = (PropertyChangedEvent.Property != NULL) ? PropertyChangedEvent.Property->GetFName() : NAME_None;
-
-	if (PropertyName == GET_MEMBER_NAME_CHECKED(FLensCaptureSettings, bShowOverlay))
+	else
 	{
-		if (TSharedPtr<FCameraCalibrationStepsController> StepsController = WeakStepsController.Pin())
-		{
-			StepsController->SetOverlayEnabled(Tool->CaptureSettings.bShowOverlay);
-		}
+		const FName AlgoName(*SelectedAlgo);
+		Tool->SetAlgo(AlgoName);
 	}
-	else if (PropertyName == GET_MEMBER_NAME_CHECKED(FLensCaptureSettings, Calibrator))
+
+	// Get the new algo
+	UCameraLensDistortionAlgo* Algo = Tool->GetAlgo();
+
+	// If we didn't change the algo, we're done here.
+	if (Algo == OldAlgo)
 	{
-		Tool->SetCalibrator(Tool->CaptureSettings.Calibrator.Get());
+		return;
 	}
-	else if (PropertyName == GET_MEMBER_NAME_CHECKED(FLensCaptureSettings, CalibrationPattern))
+
+	// Remove old UI
+	check(UI.IsValid());
+	UI->ClearChildren();
+
+	// nullptr may indicate that it was unregistered, so refresh combobox options.
+	if (!Algo)
 	{
-		DatasetListHeader->SetShowGeneratedColumn("Image", (Tool->CaptureSettings.CalibrationPattern != ECalibrationPattern::Points));
-		DatasetListHeader->SetShowGeneratedColumn("ImagePoint", (Tool->CaptureSettings.CalibrationPattern == ECalibrationPattern::Points));
-		DatasetListHeader->SetShowGeneratedColumn("ObjectPoint", (Tool->CaptureSettings.CalibrationPattern == ECalibrationPattern::Points));
+		UpdateAlgosOptions();
+		return;
 	}
+
+	// Assign GUI
+	UI->AddSlot() [Algo->BuildUI()];
 }
 
-FReply SLensDistortionToolPanel::OnClearCalibrationRowsClicked()
+void SLensDistortionToolPanel::UpdateAlgosOptions()
 {
-	if (Tool.IsValid())
+	CurrentAlgos.Empty();
+
+	// Ask the Tool for the list of registered Algos
+
+	if (!Tool.IsValid())
 	{
-		Tool->ClearCalibrationRows();
+		return;
 	}
-	return FReply::Handled();
-}
 
-FReply SLensDistortionToolPanel::OnImportDatasetClicked()
-{
-	if (Tool.IsValid())
+	for (FName& AlgoName : Tool->GetAlgos())
 	{
-		Tool->ImportCalibrationDataset();
+		CurrentAlgos.Add(MakeShared<FString>(AlgoName.ToString()));
 	}
-	return FReply::Handled();
-}
 
-FReply SLensDistortionToolPanel::OnCalibrateClicked()
-{
-	if (Tool.IsValid())
+	// Pick the first option by default (if available)
+	if (CurrentAlgos.Num())
 	{
-		Tool->CalibrateLens();
+		AlgosComboBox->SetSelectedItem(CurrentAlgos[0]);
 	}
-	return FReply::Handled();
-}
-
-TSharedRef<SWidget> SLensDistortionToolPanel::BuildDatasetListView()
-{
-	if (!Tool.IsValid())
+	else
 	{
-		return SNullWidget::NullWidget;
+		AlgosComboBox->SetSelectedItem(nullptr);
 	}
 
-	TArray<FName> HiddenColumns = { "ImagePoint", "ObjectPoint" };
-
-	DatasetListHeader = SNew(SHeaderRow)
-		.HiddenColumnsList(HiddenColumns)
-
-		+ SHeaderRow::Column("Index")
-		.DefaultLabel(LOCTEXT("IndexHeaderLabel", "Index"))
-		.FillWidth(0.2f)
-
-		+ SHeaderRow::Column("Image")
-		.DefaultLabel(LOCTEXT("ImageHeaderLabel", "Image"))
-		.FillWidth(0.8f)
-
-		+ SHeaderRow::Column("ImagePoint")
-		.DefaultLabel(LOCTEXT("ImagePointHeaderLabel", "Pixel Location"))
-		.FillWidth(0.4f)
-
-		+ SHeaderRow::Column("ObjectPoint")
-		.DefaultLabel(LOCTEXT("ObjectPointHeaderLabel", "World Position"))
-		.FillWidth(0.4f);
-
-	DatasetListView = SNew(SListView<TSharedPtr<FCalibrationRow>>)
-		.ListItemsSource(&Tool->Dataset.CalibrationRows)
-		.SelectionMode(ESelectionMode::Multi)
-		.HeaderRow(DatasetListHeader.ToSharedRef())
-		.OnGenerateRow(this, &SLensDistortionToolPanel::OnGenerateDatasetRow)
-		.OnKeyDownHandler(this, &SLensDistortionToolPanel::OnDatasetRowKeyPressed);
-
-	return DatasetListView.ToSharedRef();
+	// Ask the ComboBox to refresh its options from its source (that we just updated)
+	AlgosComboBox->RefreshOptions();
 }
 
-TSharedRef<ITableRow> SLensDistortionToolPanel::OnGenerateDatasetRow(TSharedPtr<FCalibrationRow> InItem, const TSharedRef<STableViewBase>& OwnerTable)
+TSharedRef<SWidget> SLensDistortionToolPanel::BuildSolverPickerWidget()
 {
-	return SNew(SCalibrationDatasetRow, OwnerTable, InItem);
+	return SNew(SHorizontalBox)
+
+		+ SHorizontalBox::Slot()
+		.AutoWidth()
+		.Padding(FMargin(0.0f, 0.0f, 0.0f, 0.0f))
+		[
+			SNew(SComboButton)
+			.OnGetMenuContent(this, &SLensDistortionToolPanel::BuildSolverList)
+			.ContentPadding(FMargin(4.0, 2.0))
+			.ButtonContent()
+			[
+				SNew(STextBlock).Text(this, &SLensDistortionToolPanel::GetSelectedSolverName)
+			]
+		];
 }
 
-FReply SLensDistortionToolPanel::OnDatasetRowKeyPressed(const FGeometry& Geometry, const FKeyEvent& KeyEvent)
+TSharedRef<SWidget> SLensDistortionToolPanel::BuildSolverList() 
 {
-	if (!Tool.IsValid())
-	{
-		return FReply::Unhandled();
-	}
+	// Generate menu
+	FMenuBuilder MenuBuilder(true, nullptr);
+	MenuBuilder.BeginSection("DistortionSolverSection", LOCTEXT("DistortionSolverSectionText", "Distortion Solvers"));
 
-	if ((KeyEvent.GetKey() == EKeys::A) && KeyEvent.GetModifierKeys().IsControlDown())
-	{
-		// Select all items
-		DatasetListView->SetItemSelection(Tool->Dataset.CalibrationRows, true);
-		return FReply::Handled();
-	}
-	else if (KeyEvent.GetKey() == EKeys::Escape)
-	{
-		// De-select all items
-		DatasetListView->ClearSelection();
-		return FReply::Handled();
-	}
-	else if (KeyEvent.GetKey() == EKeys::Delete)
-	{
-		// Delete selected items
-		const TArray<TSharedPtr<FCalibrationRow>> SelectedItems = DatasetListView->GetSelectedItems();
+	TArray<UClass*> DerivedSolverClasses;
+	GetDerivedClasses(ULensDistortionSolver::StaticClass(), DerivedSolverClasses);
 
-		for (const TSharedPtr<FCalibrationRow>& SelectedItem : SelectedItems)
+	if (!DerivedSolverClasses.IsEmpty())
+	{
+		for (UClass* SolverClass : DerivedSolverClasses)
 		{
-			Tool->Dataset.CalibrationRows.Remove(SelectedItem);
-			Tool->DeleteExportedRow(SelectedItem->Index);
+			ULensDistortionSolver* SolverCDO = Cast<ULensDistortionSolver>(SolverClass->GetDefaultObject());
+
+			if (SolverCDO->IsEnabled())
+			{
+				MenuBuilder.AddMenuEntry(
+					SolverCDO->GetDisplayName(),
+					SolverCDO->GetDisplayName(), // TODO: If this remains a user-facing option, then it may be better to have the solver implement a class description text field to use as the tooltip instead
+					FSlateIcon(),
+					FUIAction(
+						FExecuteAction::CreateSP(this, &SLensDistortionToolPanel::OnSolverClassSelected, SolverClass),
+						FCanExecuteAction(),
+						FIsActionChecked::CreateSP(this, &SLensDistortionToolPanel::IsSolverClassSelected, SolverClass)
+					),
+					NAME_None,
+					EUserInterfaceActionType::RadioButton
+				);
+			}
 		}
+	}
 
-		Tool->RefreshCoverage();
+	MenuBuilder.EndSection();
 
-		return FReply::Handled();
-	}
+	return MenuBuilder.MakeWidget();
 
-	return FReply::Unhandled();
 }
 
-void SLensDistortionToolPanel::RefreshListView()
+void SLensDistortionToolPanel::OnSolverClassSelected(UClass* SolverClass)
+{
+	Tool->SetSolverClass(SolverClass);
+}
+
+bool SLensDistortionToolPanel::IsSolverClassSelected(UClass* SolverClass) const
 {
-	if (DatasetListView)
+	return (Tool->GetSolverClass() == SolverClass);
+}
+
+FText SLensDistortionToolPanel::GetSelectedSolverName() const
+{
+	if (UClass* SolverClass = Tool->GetSolverClass())
 	{
-		DatasetListView->RequestListRefresh();
+		ULensDistortionSolver* SolverCDO = Cast<ULensDistortionSolver>(SolverClass->GetDefaultObject());
+		return SolverCDO->GetDisplayName();
 	}
+	
+	return FText::GetEmpty();
 }
 
-void SLensDistortionToolPanel::BuildProgressWindow()
+TSharedRef<SWidget> SLensDistortionToolPanel::BuildAlgoPickerWidget()
 {
-	ProgressWindow = SNew(SWindow)
-		.Title(LOCTEXT("ProgressWindowTitle", "Distortion Calibration Progress"))
-		.SizingRule(ESizingRule::Autosized)
-		.IsTopmostWindow(true)
-		.HasCloseButton(false)
-		.SupportsMaximize(false)
-		.SupportsMinimize(true);
-
-	ProgressTextWidget = SNew(STextBlock).Text(FText::GetEmpty());
-
-	OkayButton = SNew(SButton)
-		.IsEnabled(false)
-		.HAlign(HAlign_Center)
-		.VAlign(VAlign_Center)
-		.Text(LOCTEXT("OkText", "Ok"))
-		.OnClicked(this, &SLensDistortionToolPanel::OnOkPressed);
-
-	TSharedRef<SWidget> WindowContent = SNew(SVerticalBox)
-
-		// Text widget to display the current progress of the calibration
-		+ SVerticalBox::Slot()
-		.HAlign(EHorizontalAlignment::HAlign_Center)
-		.VAlign(EVerticalAlignment::VAlign_Center)
-		[
-			ProgressTextWidget.ToSharedRef()
-		]
+	// Create ComboBox widget
 
-		// Ok and Cancel buttons
-		+ SVerticalBox::Slot()
-		.HAlign(EHorizontalAlignment::HAlign_Center)
-		.VAlign(EVerticalAlignment::VAlign_Center)
+	AlgosComboBox = SNew(SComboBox<TSharedPtr<FString>>)
+		.OptionsSource(&CurrentAlgos)
+		.OnSelectionChanged_Lambda([&](TSharedPtr<FString> NewValue, ESelectInfo::Type Type) -> void
+		{
+			// Replace the custom algo widget
+			UpdateUI();
+		})
+		.OnGenerateWidget_Lambda([&](TSharedPtr<FString> InOption) -> TSharedRef<SWidget>
+		{
+			return SNew(STextBlock).Text(FText::FromString(*InOption));
+		})
+		.InitiallySelectedItem(nullptr)
 		[
-			SNew(SHorizontalBox)
-
-			+ SHorizontalBox::Slot()
-			.AutoWidth()
-			[
-				OkayButton.ToSharedRef()
-			]
-
-			+ SHorizontalBox::Slot()
-			.AutoWidth()
-			[
-				SNew(SButton)
-				.Text(LOCTEXT("CancelText", "Cancel"))
-				.HAlign(HAlign_Center)
-				.VAlign(VAlign_Center)
-				.OnClicked(this, &SLensDistortionToolPanel::OnCancelPressed)
-			]
+			SNew(STextBlock)
+			.Text_Lambda([&]() -> FText
+			{
+				if (AlgosComboBox.IsValid() && AlgosComboBox->GetSelectedItem().IsValid())
+				{
+					return FText::FromString(*AlgosComboBox->GetSelectedItem());
+				}
+
+				return LOCTEXT("InvalidComboOption", "Invalid");
+			})
 		];
 
-	ProgressWindow->SetContent(WindowContent);
+	// Update the object holding this combobox's options source
+	UpdateAlgosOptions();
 
-	// Create the window, but start with it hidden. When the user initiates a calibration, the progress window will be shown.
-	FSlateApplication::Get().AddWindow(ProgressWindow.ToSharedRef());
-	ProgressWindow->HideWindow();
-}
-
-FReply SLensDistortionToolPanel::OnCancelPressed()
-{
-	if (Tool.IsValid())
+	// Pick the first option by default (if available)
+	if (CurrentAlgos.Num())
 	{
-		Tool->CancelCalibration();
-		Tool->CalibrationTask = {};
+		AlgosComboBox->SetSelectedItem(CurrentAlgos[0]);
 	}
-
-	ProgressWindow->HideWindow();
-	SetEnabled(true);
-
-	return FReply::Handled();
-}
-
-FReply SLensDistortionToolPanel::OnOkPressed()
-{
-	if (Tool.IsValid())
+	else
 	{
-		Tool->SaveCalibrationResult();
+		AlgosComboBox->SetSelectedItem(nullptr);
 	}
 
-	ProgressWindow->HideWindow();
-	SetEnabled(true);
+	return SNew(SHorizontalBox)
+
+		+ SHorizontalBox::Slot() // algo picker
+		[
+			AlgosComboBox.ToSharedRef()
+		]
 
-	return FReply::Handled();
+		+ SHorizontalBox::Slot() // Help button
+		.AutoWidth()
+		[
+			SNew(SButton)
+			.ToolTipText(LOCTEXT("ShowHelp_Tip", "Help about this algo"))
+			.HAlign(HAlign_Center)
+			.VAlign(VAlign_Center)
+			.ButtonStyle(FAppStyle::Get(), "HoverHintOnly")
+			.OnClicked_Lambda([&]() -> FReply
+			{
+				if (!Tool.IsValid())
+				{
+					return FReply::Handled();
+				}
+
+				UCameraLensDistortionAlgo* Algo = Tool->GetAlgo();
+
+				if (!Algo)
+				{
+					return FReply::Handled();
+				}
+
+				TSharedRef< SCustomDialog> AlgoHelpWindow = 
+					SNew(SCustomDialog)
+					.Title(FText::FromName(Tool->FriendlyName()))
+					.Content()
+					[
+						Algo->BuildHelpWidget()
+					]
+					.Buttons({
+						SCustomDialog::FButton(LOCTEXT("Ok", "Ok")),
+					});
+
+				AlgoHelpWindow->Show();
+
+				return FReply::Handled();
+			})
+			[
+				SNew(STextBlock)
+				.Font(FAppStyle::Get().GetFontStyle("FontAwesome.12"))
+				.Text(FEditorFontGlyphs::Info_Circle)
+				.ColorAndOpacity(FLinearColor::White)
+			]
+		]
+		;
 }
 
+
 #undef LOCTEXT_NAMESPACE
diff --git a/Engine/Plugins/VirtualProduction/CameraCalibration/Source/CameraCalibrationEditor/Private/AssetEditor/SLensDistortionToolPanel.h b/Engine/Plugins/VirtualProduction/CameraCalibration/Source/CameraCalibrationEditor/Private/AssetEditor/SLensDistortionToolPanel.h
--- a/Engine/Plugins/VirtualProduction/CameraCalibration/Source/CameraCalibrationEditor/Private/AssetEditor/SLensDistortionToolPanel.h
+++ b/Engine/Plugins/VirtualProduction/CameraCalibration/Source/CameraCalibrationEditor/Private/AssetEditor/SLensDistortionToolPanel.h
@@ -1,103 +1,74 @@
 // Copyright Epic Games, Inc. All Rights Reserved.
 #pragma once
 
-#include "Widgets/SCompoundWidget.h"
+#include "CoreMinimal.h"
 
-#include "CameraCalibrationStepsController.h"
-#include "IDetailCustomization.h"
-#include "LensDistortionTool.h"
-#include "PropertyHandle.h"
-#include "Widgets/SWindow.h"
-#include "Widgets/Text/STextBlock.h"
-#include "Widgets/Views/SHeaderRow.h"
-#include "Widgets/Views/SListView.h"
-#include "Widgets/Views/STableRow.h"
+#include "Widgets/DeclarativeSyntaxSupport.h"
+#include "Widgets/SCompoundWidget.h"
+#include "UObject/StrongObjectPtr.h"
 
+class FString;
+class ULensFile;
 class ULensDistortionTool;
 
-struct FCalibrationRow;
-
-/** Customization for the FCaptureSettings struct */
-class FCaptureSettingsCustomization : public IDetailCustomization
-{
-public:
-	virtual void CustomizeDetails(IDetailLayoutBuilder& DetailBuilder) override;
-
-private:
-	TSharedPtr<IPropertyHandle> PropertyHandle;
-	TSharedPtr<IPropertyHandle> NextPointPropertyHandle;
-
-	FText GetNextPointName() const;
-
-	void OnCalibratorSelected(const FAssetData& AssetData);
+template<typename OptionType>
+class SComboBox;
 
-	bool DoesAssetHaveCalibrationComponent(const FAssetData& AssetData) const;
+struct FArguments;
 
-	FAssetData GetCalibratorAssetData() const;
-};
-
-/** Row widget for the calibration row list view */
-class SCalibrationDatasetRow : public SMultiColumnTableRow<TSharedPtr<FCalibrationRow>>
+/**
+ * Wrapper UI for the specified calibration step.
+ */
+class SLensDistortionToolPanel : public SCompoundWidget
 {
-public:
-	SLATE_BEGIN_ARGS(SCalibrationDatasetRow) {}
+	SLATE_BEGIN_ARGS(SLensDistortionToolPanel) {}
 	SLATE_END_ARGS()
 
-	void Construct(const FArguments& Args, const TSharedRef<STableViewBase>& OwnerTableView, TSharedPtr<FCalibrationRow>& InRowData);
+public:
+	void Construct(const FArguments& InArgs, ULensDistortionTool* InTool);
 
-	virtual TSharedRef<SWidget> GenerateWidgetForColumn(const FName& ColumnName) override;
+	/** Refreshes the list of available algorithms shown in the AlgosComboBox */
+	void UpdateAlgosOptions();
 
 private:
-	TSharedPtr<FCalibrationRow> RowData;
-};
 
-/** UI for the Lens Distortion step */
-class SLensDistortionToolPanel : public SCompoundWidget
-{
-	SLATE_BEGIN_ARGS(SLensDistortionToolPanel) {}
-	SLATE_END_ARGS()
+	/** Builds the wrapper for the currently selected algo UI */
+	TSharedRef<SWidget> BuildUIWrapper();
 
-public:
-	void Construct(const FArguments& InArgs, ULensDistortionTool* InTool, TWeakPtr<FCameraCalibrationStepsController> InStepsController);
+	/** Builds the UI of the solver picker widget */
+	TSharedRef<SWidget> BuildSolverPickerWidget();
 
-	/** Refresh the dataset list view */
-	void RefreshListView();
+	/** Builds the list of solvers to display when opening the solver picker */
+	TSharedRef<SWidget> BuildSolverList();
 
-	/** Hides the modal progress window before the rest of the UI is destroyed */
-	void Shutdown();
+	/** When the solver class is selected in the UI, update the lens distortion tool with the newly selected class */
+	void OnSolverClassSelected(UClass* SolverClass);
 
-	/** Shows the modal progress window */
-	void OpenProgressWindow();
+	/** Returns true if the input class is the currently selected one */
+	bool IsSolverClassSelected(UClass* SolverClass) const;
 
-	/** Enables the OkayButton in the progress window */
-	void MarkProgressFinished();
+	/** Get the display name of the currently selected solver class */
+	FText GetSelectedSolverName() const;
 
-	void UpdateProgressText(const FText& ProgressText);
+	/** Builds the UI for the algorithm picker */
+	TSharedRef<SWidget> BuildAlgoPickerWidget();
 
-private:
-	/** Callback to handle changes to the customized capture settings */
-	void OnCaptureSettingsChanged(const FPropertyChangedEvent& PropertyChangedEvent);
+	/** Updates the UI so that it matches the selected algorithm (if necessary) */
+	void UpdateUI();
 
-	TSharedRef<SWidget> BuildDatasetListView();
-	TSharedRef<ITableRow> OnGenerateDatasetRow(TSharedPtr<FCalibrationRow> InItem, const TSharedRef<STableViewBase>& OwnerTable);
-	FReply OnDatasetRowKeyPressed(const FGeometry& Geometry, const FKeyEvent& KeyEvent);
+private:
 
-	FReply OnClearCalibrationRowsClicked();
-	FReply OnImportDatasetClicked();
-	FReply OnCalibrateClicked();
+	/** The tool controller object */
+	TWeakObjectPtr<class ULensDistortionTool> Tool;
 
-	void BuildProgressWindow();
-	FReply OnOkPressed();
-	FReply OnCancelPressed();
+	/** The box containing the UI given by the selected algorithm */
+	TSharedPtr<class SVerticalBox> UI;
 	
 private:
-	TWeakObjectPtr<ULensDistortionTool> Tool;
-	TWeakPtr<FCameraCalibrationStepsController> WeakStepsController;
 
-	TSharedPtr<SListView<TSharedPtr<FCalibrationRow>>> DatasetListView;
-	TSharedPtr<SHeaderRow> DatasetListHeader;
+	/** Options source for the AlgosComboBox. Lists the currently available nodal offset algorithms */
+	TArray<TSharedPtr<FString>> CurrentAlgos;
 
-	TSharedPtr<SWindow> ProgressWindow;
-	TSharedPtr<STextBlock> ProgressTextWidget;
-	TSharedPtr<SButton> OkayButton;
+	/** The combobox that presents the available nodal offset algorithms */
+	TSharedPtr<SComboBox<TSharedPtr<FString>>> AlgosComboBox;
 };
diff --git a/Engine/Plugins/VirtualProduction/CameraCalibration/Source/CameraCalibrationEditor/Private/Calibrators/CameraCalibrationSolver.h b/Engine/Plugins/VirtualProduction/CameraCalibration/Source/CameraCalibrationEditor/Private/Calibrators/CameraCalibrationSolver.h
--- a/Engine/Plugins/VirtualProduction/CameraCalibration/Source/CameraCalibrationEditor/Private/Calibrators/CameraCalibrationSolver.h
+++ b/Engine/Plugins/VirtualProduction/CameraCalibration/Source/CameraCalibrationEditor/Private/Calibrators/CameraCalibrationSolver.h
@@ -141,7 +141,7 @@ protected:
   * The implementation is largely based on the implementation of calibrateCamera from OpenCV: [REDACTED_PROFILE_URL]/opencv
   */
 
-UCLASS(meta = (DisplayName="OpenCV Solver"))
+UCLASS()
 class ULensDistortionSolverOpenCV : public ULensDistortionSolver
 {
 	GENERATED_BODY()
diff --git a/Engine/Plugins/VirtualProduction/CameraCalibration/Source/CameraCalibrationEditor/Private/Calibrators/CameraLensDistortionAlgoAruco.cpp b/Engine/Plugins/VirtualProduction/CameraCalibration/Source/CameraCalibrationEditor/Private/Calibrators/CameraLensDistortionAlgoAruco.cpp
new file mode 100644
--- /dev/null
+++ b/Engine/Plugins/VirtualProduction/CameraCalibration/Source/CameraCalibrationEditor/Private/Calibrators/CameraLensDistortionAlgoAruco.cpp
@@ -0,0 +1,1061 @@
+// Copyright Epic Games, Inc. All Rights Reserved.
+
+#include "CameraLensDistortionAlgoAruco.h"
+
+#include "AssetEditor/CameraCalibrationStepsController.h"
+#include "AssetEditor/LensDistortionTool.h"
+#include "CalibrationPointComponent.h"
+#include "Camera/CameraActor.h"
+#include "CameraCalibrationEditorLog.h"
+#include "CameraCalibrationSolver.h"
+#include "CameraCalibrationUtilsPrivate.h"
+#include "Dialog/SCustomDialog.h"
+#include "Dom/JsonObject.h"
+#include "Engine/Texture2D.h"
+#include "Framework/MultiBox/MultiBoxBuilder.h"
+#include "ImageUtils.h"
+#include "JsonObjectConverter.h"
+#include "LensComponent.h"
+#include "LensFile.h"
+#include "Misc/MessageDialog.h"
+#include "Models/AnamorphicLensModel.h"
+#include "Models/SphericalLensModel.h"
+#include "OpenCVHelper.h"
+#include "UI/CameraCalibrationWidgetHelpers.h"
+#include "Widgets/Input/SButton.h"
+#include "Widgets/Input/SCheckBox.h"
+#include "Widgets/Input/SComboBox.h"
+#include "Widgets/Input/SNumericEntryBox.h"
+#include "Widgets/Views/SListView.h"
+#include "Widgets/Views/STableRow.h"
+#include "UI/SFilterableActorPicker.h"
+#include "UI/SImageTexture.h"
+
+
+#define LOCTEXT_NAMESPACE "CameraLensDistortionAlgoAruco"
+
+#if WITH_EDITOR
+static TAutoConsoleVariable<bool> CVarFixExtrinsicsAruco(TEXT("LensDistortionAruco.FixExtrinsics"), false, TEXT("If true, the solver will fix the camera extrinsics to the user-provided camera poses"));
+static TAutoConsoleVariable<bool> CVarFixDistortionAruco(TEXT("LensDistortionAruco.FixDistortion"), false, TEXT("If true, the solver will not optimize distortion, and will use the input distortion values if any are given, or assume zero distortion otherwise."));
+static TAutoConsoleVariable<bool> CVarUseExtrinsicsGuessAruco(TEXT("LensDistortionAruco.UseExtrinsicsGuess"), false, TEXT("If true, the actual calibrator and camera poses will be used when running the solver"));
+static TAutoConsoleVariable<bool> CVarSolveTargetOffsetAruco(TEXT("LensDistortionAruco.SolveTargetOffset"), false, TEXT("If true, the solver will calibrate for an offset from the calibrator tracking origin that further minimizes reprojection error."));
+#endif
+
+const int UCameraLensDistortionAlgoAruco::DATASET_VERSION = 1;
+
+// String constants for import/export
+namespace UE::CameraCalibration::Private::LensDistortionArucoExportFields
+{
+	static const FString Version(TEXT("Version"));
+}
+
+namespace UE::CameraCalibration::Private
+{
+	// Test the input array of weak pointers for validity, and return a list of valid pointers to calibration point components
+	void GetValidActiveComponents(const TArray<TWeakObjectPtr<UCalibrationPointComponent>>& WeakComponents, TArray<UCalibrationPointComponent*>& ValidComponents)
+	{
+		ValidComponents.Reserve(WeakComponents.Num());
+		for (const TWeakObjectPtr<UCalibrationPointComponent>& WeakComponent : WeakComponents)
+		{
+			if (UCalibrationPointComponent* Component = WeakComponent.Get())
+			{
+				ValidComponents.Add(Component);
+			}
+		}
+	}
+
+	// Generate a debug texture with the aruco points from the input row drawn onto it
+	UTexture2D* CreateDebugTexture(TSharedPtr<FLensDistortionArucoRowData> Row)
+	{
+		if (!Row)
+		{
+			return nullptr;
+		}
+
+		TArray<FArucoMarker> DebugMarkers;
+		DebugMarkers.Reserve(Row->ArucoPoints.Num());
+
+		// Extract the data from the aruco points that OpenCV expects to draw the debug markers
+		for (int32 MarkerIndex = 0; MarkerIndex < Row->ArucoPoints.Num(); ++MarkerIndex)
+		{
+			const FArucoCalibrationPoint& Marker = Row->ArucoPoints[MarkerIndex];
+
+			FArucoMarker NewMarker;
+			NewMarker.MarkerID = Marker.MarkerID;
+			for (int32 CornerIndex = 0; CornerIndex < 4; ++CornerIndex)
+			{
+				NewMarker.Corners[CornerIndex] = Marker.Corners2D[CornerIndex];
+			}
+
+			DebugMarkers.Add(NewMarker);
+		}
+
+		UTexture2D* DebugTexture = FImageUtils::CreateTexture2DFromImage(Row->MediaImage);
+		FOpenCVHelper::DrawArucoMarkers(DebugMarkers, DebugTexture);
+
+		return DebugTexture;
+	}
+
+	// Open a dialog to display the input image 
+	void DisplayDebugImageDialog(UTexture2D* Image)
+	{
+		if (Image)
+		{
+			TSharedRef<SCustomDialog> DebugImageDialog =
+				SNew(SCustomDialog)
+				.UseScrollBox(false)
+				.Content()
+				[
+					SNew(SImageTexture, Image)
+				]
+				.Buttons
+				({
+					SCustomDialog::FButton(LOCTEXT("OkButton", "Ok")),
+				});
+
+			DebugImageDialog->Show();
+		}
+	}
+}
+
+namespace CameraLensDistortionAlgoAruco
+{
+	class SCalibrationRowGenerator : public SMultiColumnTableRow<TSharedPtr<FLensDistortionArucoRowData>>
+	{
+
+	public:
+		SLATE_BEGIN_ARGS(SCalibrationRowGenerator) {}
+
+		/** The list item for this row */
+		SLATE_ARGUMENT(TSharedPtr<FLensDistortionArucoRowData>, CalibrationRowData)
+
+		SLATE_END_ARGS()
+
+		void Construct(const FArguments& Args, const TSharedRef<STableViewBase>& OwnerTableView)
+		{
+			CalibrationRowData = Args._CalibrationRowData;
+
+			SMultiColumnTableRow<TSharedPtr<FLensDistortionArucoRowData>>::Construct(
+				FSuperRowType::FArguments()
+				.Padding(1.0f),
+				OwnerTableView
+			);
+		}
+
+		//~Begin SMultiColumnTableRow
+		virtual TSharedRef<SWidget> GenerateWidgetForColumn(const FName& ColumnName) override
+		{
+			if (!CalibrationRowData)
+			{
+				return SNullWidget::NullWidget;
+			}
+
+			if (ColumnName == TEXT("Index"))
+			{
+				return SNew(SVerticalBox)
+					+ SVerticalBox::Slot()
+					.HAlign(EHorizontalAlignment::HAlign_Center)
+					.VAlign(EVerticalAlignment::VAlign_Center)
+					[SNew(STextBlock).Text(FText::AsNumber(CalibrationRowData->Index))];
+			}
+
+			if (ColumnName == TEXT("NumMarkers"))
+			{
+				return SNew(SVerticalBox)
+					+ SVerticalBox::Slot()
+					.HAlign(EHorizontalAlignment::HAlign_Center)
+					.VAlign(EVerticalAlignment::VAlign_Center)
+					[SNew(STextBlock).Text(FText::AsNumber(CalibrationRowData->ArucoPoints.Num()))];
+			}
+
+			if (ColumnName == TEXT("Image"))
+			{
+				// Generate a transient thumbnail texture to display in the tool
+				const FImage& MediaImage = CalibrationRowData->MediaImage;
+
+				constexpr int32 ResolutionDivider = 4;
+				const FIntPoint ThumbnailSize = FIntPoint(MediaImage.SizeX / ResolutionDivider, MediaImage.SizeY / ResolutionDivider);
+
+				FImage ThumbnailImage;
+				FImageCore::ResizeTo(MediaImage, ThumbnailImage, ThumbnailSize.X, ThumbnailSize.Y, MediaImage.Format, MediaImage.GetGammaSpace());
+
+				if (UTexture2D* Thumbnail = FImageUtils::CreateTexture2DFromImage(ThumbnailImage))
+				{
+					return SNew(SImageTexture, Thumbnail)
+						.MinDesiredHeight(4 * FCameraCalibrationWidgetHelpers::DefaultRowHeight)
+						.MaxDesiredHeight(4 * FCameraCalibrationWidgetHelpers::DefaultRowHeight);
+				}
+				else
+				{
+					const FString Text = FString::Printf(TEXT("Image Unavailable"));
+					return SNew(STextBlock).Text(FText::FromString(Text));
+				}
+			}
+
+			return SNullWidget::NullWidget;
+		}
+		//~End SMultiColumnTableRow
+
+
+	private:
+		TSharedPtr<FLensDistortionArucoRowData> CalibrationRowData;
+	};
+
+};
+
+void UCameraLensDistortionAlgoAruco::Initialize(ULensDistortionTool* InTool)
+{
+	WeakTool = InTool;
+
+	// Find all actors in the current level that have calibration point components
+	TArray<AActor*> CalibratorActors;
+	UE::CameraCalibration::Private::FindActorsWithCalibrationComponents(CalibratorActors);
+
+	// If at least one actor was found, set the initial calibrator to the first one in the list
+	if (CalibratorActors.Num() > 0)
+	{
+		SetCalibrator(CalibratorActors[0]);
+	}
+}
+
+void UCameraLensDistortionAlgoAruco::Shutdown()
+{
+	WeakTool.Reset();
+	WeakCalibrator.Reset();
+}
+
+bool UCameraLensDistortionAlgoAruco::SupportsModel(const TSubclassOf<ULensModel>& LensModel) const
+{
+	return ((LensModel == USphericalLensModel::StaticClass()) || (LensModel == UAnamorphicLensModel::StaticClass()));
+}
+
+void UCameraLensDistortionAlgoAruco::Tick(float DeltaTime)
+{
+}
+
+bool UCameraLensDistortionAlgoAruco::OnViewportClicked(const FGeometry& MyGeometry, const FPointerEvent& MouseEvent)
+{
+	// We only respond to left clicks
+	if (MouseEvent.GetEffectingButton() != EKeys::LeftMouseButton)
+	{
+		return false;
+	}
+
+	// Attempt to add a new calibration row with data from the current media image, and display an error dialog to the user if any errors occur
+	FText ErrorMessage;
+	if (!AddCalibrationRow(ErrorMessage))
+	{
+		if (ErrorMessage.IsEmpty())
+		{
+			ErrorMessage = LOCTEXT("UnknownCalibrationError", "An error occurred. Please check the output log for more details.");
+		}
+
+		FMessageDialog::Open(EAppMsgType::Ok, ErrorMessage);
+	}
+
+	// If the media was currently paused, then resume it after gathering data from the current image
+	if (ULensDistortionTool* LensDistortionTool = WeakTool.Get())
+	{
+		if (FCameraCalibrationStepsController* StepsController = LensDistortionTool->GetCameraCalibrationStepsController())
+		{
+			StepsController->Play();
+		}
+	}
+
+	return true;
+}
+
+FDistortionCalibrationTask UCameraLensDistortionAlgoAruco::BeginCalibration(FText& OutErrorMessage)
+{
+	FDistortionCalibrationTask CalibrationTask = {};
+
+	if (CalibrationRows.Num() < 1)
+	{
+		OutErrorMessage = LOCTEXT("NotEnoughCalibrationRowsError", "Could not initiate distortion calibration. At least 1 calibration row is required.");
+		return CalibrationTask;
+	}
+
+	ULensDistortionTool* LensDistortionTool = WeakTool.Get();
+	if (!ensureMsgf((LensDistortionTool != nullptr), TEXT("The Lens Distortion Tool was invalid.")))
+	{
+		return CalibrationTask;
+	}
+
+	FCameraCalibrationStepsController* StepsController = LensDistortionTool->GetCameraCalibrationStepsController();
+	if (!ensureMsgf((StepsController != nullptr), TEXT("The Calibration Steps Controller was invalid.")))
+	{
+		return CalibrationTask;
+	}
+
+	ULensFile* LensFile = StepsController->GetLensFile();
+	if (!ensureMsgf((LensFile != nullptr), TEXT("The Lens File was invalid.")))
+	{
+		return CalibrationTask;
+	}
+
+	// The evaluation focus and zoom values should not change during data collection, so all rows should have the same values.
+	const FLensFileEvaluationInputs LensFileEvalInputs = StepsController->GetLensFileEvaluationInputs();
+	const float Focus = LensFileEvalInputs.Focus;
+	const float Zoom = LensFileEvalInputs.Zoom;
+
+	// Initialize the focal length estimate in pixels
+	const FIntPoint ImageSize = LensFile->CameraFeedInfo.GetDimensions();
+
+	const float PixelAspect = LensFile->LensInfo.SqueezeFactor;
+
+	if (FMath::IsNearlyZero(PixelAspect))
+	{
+		OutErrorMessage = LOCTEXT("PixelAspectZeroError", "The pixel aspect ratio of the CineCamera is zero, which is invalid.");
+		return CalibrationTask;
+	}
+
+	const float PhysicalSensorWidth = StepsController->GetLensFileEvaluationInputs().Filmback.SensorWidth;
+	const float DesqueezeSensorWidth = PhysicalSensorWidth * PixelAspect;
+
+	if (FMath::IsNearlyZero(DesqueezeSensorWidth))
+	{
+		OutErrorMessage = LOCTEXT("SensorWidthZeroError", "One of the filmback dimensions of the CineCamera is zero, which is invalid.");
+		return CalibrationTask;
+	}
+
+	if (!FocalLengthEstimate.IsSet() || FMath::IsNearlyZero(FocalLengthEstimate.GetValue()))
+	{
+		OutErrorMessage = LOCTEXT("FocalLengthEstimateError", "Enter a non-zero value (in mm) for the focal length estimate.");
+		return CalibrationTask;
+	}
+
+	const float FocalLengthEstimateValue = FocalLengthEstimate.GetValue();
+	const double Fx = (FocalLengthEstimateValue / DesqueezeSensorWidth) * ImageSize.X;
+
+	// When operating on a desqueezed image, we expect our pixel aspect to be square, so horizontal and vertical field of view are assumed to be equal (i.e. Fx == Fy)
+	FVector2D FocalLength = FVector2D(Fx, Fx);
+	FVector2D ImageCenter = FVector2D((ImageSize.X - 1) * 0.5, (ImageSize.Y - 1) * 0.5);
+
+	TArray<FObjectPoints> Samples3d;
+	Samples3d.Reserve(CalibrationRows.Num());
+
+	TArray<FImagePoints> Samples2d;
+	Samples2d.Reserve(CalibrationRows.Num());
+
+	TArray<FTransform> CameraPoses;
+	CameraPoses.Reserve(CalibrationRows.Num());
+
+	TArray<FTransform> TargetPoses;
+	CameraPoses.Reserve(CalibrationRows.Num());
+
+	// Extract the 3D points, 2D points, and camera poses from each row to pass to the solver
+	for (const TSharedPtr<FLensDistortionArucoRowData>& Row : CalibrationRows)
+	{
+		check(Row.IsValid());
+
+		FObjectPoints Points3d;
+		Points3d.Points.Reserve(Row->ArucoPoints.Num() * 4);
+
+		FImagePoints Points2d;
+		Points2d.Points.Reserve(Row->ArucoPoints.Num() * 4);
+
+		for (const FArucoCalibrationPoint& Marker : Row->ArucoPoints)
+		{
+			for (int32 CornerIndex = 0; CornerIndex < 4; ++CornerIndex)
+			{
+				Points3d.Points.Add(Marker.Corners3D[CornerIndex]);
+
+				const FVector2f& MarkerCorner2D = Marker.Corners2D[CornerIndex];
+				Points2d.Points.Add(FVector2D(MarkerCorner2D.X, MarkerCorner2D.Y));
+			}
+		}
+
+		Samples3d.Add(Points3d);
+		Samples2d.Add(Points2d);
+		CameraPoses.Add(Row->CameraPose);
+
+		FTransform FirstMarkerPose;
+
+		FArucoCalibrationPoint& FirstAruco = Row->ArucoPoints[0];
+
+		FVector TopLeft = FirstAruco.Corners3D[0];
+		FVector TopRight = FirstAruco.Corners3D[1];
+		FVector BottomLeft = FirstAruco.Corners3D[3];
+
+		FirstMarkerPose.SetLocation(TopLeft);
+
+		FRotator FirstMarkerRotation = FRotationMatrix::MakeFromYZ(TopRight - TopLeft, TopLeft - BottomLeft).Rotator();
+		FirstMarkerPose.SetRotation(FirstMarkerRotation.Quaternion());
+
+		TargetPoses.Add(FirstMarkerPose);
+	}
+
+	ECalibrationFlags SolverFlags = ECalibrationFlags::None;
+
+	// Aruco markers may be detected anywhere in the image, and there is no guarantee that they will be coplanar. 
+	// The solver's initialization for focal length assumes all points in an image are coplanar. 
+	// Therefore, we must provide an intrinsics guess to skip this initialization step in the solver. 
+	EnumAddFlags(SolverFlags, ECalibrationFlags::UseIntrinsicGuess);
+
+	if (CVarUseExtrinsicsGuessAruco.GetValueOnGameThread())
+	{
+		EnumAddFlags(SolverFlags, ECalibrationFlags::UseExtrinsicGuess);
+	}
+
+	if (CVarFixExtrinsicsAruco.GetValueOnAnyThread())
+	{
+		EnumAddFlags(SolverFlags, ECalibrationFlags::FixExtrinsics);
+	}
+
+	if (CVarFixDistortionAruco.GetValueOnAnyThread())
+	{
+		EnumAddFlags(SolverFlags, ECalibrationFlags::FixDistortion);
+	}
+
+	if (bFixFocalLength)
+	{
+		EnumAddFlags(SolverFlags, ECalibrationFlags::FixFocalLength);
+	}
+
+	if (bFixImageCenter)
+	{
+		EnumAddFlags(SolverFlags, ECalibrationFlags::FixPrincipalPoint);
+	}
+
+	if (CVarSolveTargetOffsetAruco.GetValueOnGameThread())
+	{
+		EnumAddFlags(SolverFlags, ECalibrationFlags::SolveTargetOffset);
+	}
+
+	FDistortionInfo DistortionGuess;
+	LensFile->EvaluateDistortionParameters(Focus, Zoom, DistortionGuess);
+
+	const TSubclassOf<ULensModel> Model = LensFile->LensInfo.LensModel;
+
+	UClass* SolverClass = LensDistortionTool->GetSolverClass();
+
+	CalibrationTask = UE::Tasks::Launch(UE_SOURCE_LOCATION, [SolverClass, Model, Samples3d, Samples2d, ImageSize, FocalLength, ImageCenter, DistortionGuess, CameraPoses, TargetPoses, PixelAspect, SolverFlags, Focus, Zoom]() mutable
+		{
+			ULensDistortionSolver* Solver = NewObject<ULensDistortionSolver>(GetTransientPackage(), SolverClass);
+
+			FDistortionCalibrationResult Result = Solver->Solve(
+				Samples3d,
+				Samples2d,
+				ImageSize,
+				FocalLength,
+				ImageCenter,
+				DistortionGuess.Parameters,
+				CameraPoses,
+				TargetPoses,
+				Model,
+				PixelAspect,
+				SolverFlags
+			);
+
+			// CalibrateCamera() returns focal length and image center in pixels, but the result is expected to be normalized by the image size
+			Result.FocalLength.FxFy = Result.FocalLength.FxFy / ImageSize;
+			Result.ImageCenter.PrincipalPoint = Result.ImageCenter.PrincipalPoint / ImageSize;
+
+			// FZ inputs to LUT
+			Result.EvaluatedFocus = Focus;
+			Result.EvaluatedZoom = Zoom;
+
+			return Result;
+		});
+
+	return CalibrationTask;
+}
+
+void UCameraLensDistortionAlgoAruco::OnDistortionSavedToLens()
+{
+	// Since the calibration result was saved, there is no further use for the current samples.
+	ClearCalibrationRows();
+}
+
+bool UCameraLensDistortionAlgoAruco::AddCalibrationRow(FText& OutErrorMessage)
+{
+	// Detect the aruco dictionary to use by looking at the names of the calibration points on the selected calibrator actor
+	AActor* CalibratorActor = WeakCalibrator.Get();
+	if (!CalibratorActor)
+	{
+		OutErrorMessage = LOCTEXT("InvalidCalibrator", "Please pick a calibrator actor from the list that contains aruco calibration points.");
+		return false;
+	}
+
+	EArucoDictionary Dictionary = UE::CameraCalibration::Private::GetArucoDictionaryForCalibrator(CalibratorActor);
+	if (Dictionary == EArucoDictionary::None)
+	{
+		OutErrorMessage = LOCTEXT("InvalidArucoDictionary", "The names of the calibration point component(s) of the selected calibrator do not match any known aruco dictionaries.");
+		return false;
+	}
+
+	ULensDistortionTool* LensDistortionTool = WeakTool.Get();
+	if (!LensDistortionTool)
+	{
+		UE_LOG(LogCameraCalibrationEditor, Warning, TEXT("Failed to gather calibration data from the viewport. The Lens Distortion Tool was invalid."));
+		return false;
+	}
+
+	FCameraCalibrationStepsController* StepsController = LensDistortionTool->GetCameraCalibrationStepsController();
+	if (!StepsController)
+	{
+		UE_LOG(LogCameraCalibrationEditor, Warning, TEXT("Failed to gather calibration data from the viewport. The Steps Controller was invalid."));
+		return false;
+	}
+
+	// Read the image data from the current media source into a CPU-accessible array
+	TArray<FColor> Pixels;
+	FIntPoint Size = FIntPoint(0, 0);
+	if (!StepsController->ReadMediaPixels(Pixels, Size, OutErrorMessage, ESimulcamViewportPortion::CameraFeed))
+	{
+		return false;
+	}
+
+	// Identify any aruco markers matching the current dictionary in the media image
+	TArray<FArucoMarker> IdentifiedMarkers;
+	bool bResult = FOpenCVHelper::IdentifyArucoMarkers(Pixels, Size, Dictionary, IdentifiedMarkers);
+
+	if (!bResult || IdentifiedMarkers.IsEmpty())
+	{
+		OutErrorMessage = LOCTEXT("NoArucoMarkersFound", "No aruco markers could be identified in the media image, or they did not match the aruco dictionary of the selected calibrator actor.");
+		return false;
+	}
+
+	// For each identified marker, search the calibration components to find the subpoints matching that marker and the 3D location of each of its corners
+	TArray<FArucoCalibrationPoint> ArucoCalibrationPoints;
+	ArucoCalibrationPoints.Reserve(IdentifiedMarkers.Num());
+
+	TArray<UCalibrationPointComponent*> ValidActiveCalibrationComponents;
+	UE::CameraCalibration::Private::GetValidActiveComponents(ActiveCalibratorComponents, ValidActiveCalibrationComponents);
+
+	for (const FArucoMarker& Marker : IdentifiedMarkers)
+	{
+		FArucoCalibrationPoint ArucoCalibrationPoint;
+		if (UE::CameraCalibration::Private::FindArucoCalibrationPoint(ValidActiveCalibrationComponents, Dictionary, Marker, ArucoCalibrationPoint))
+		{
+			ArucoCalibrationPoints.Add(ArucoCalibrationPoint);
+		}
+	}
+
+	ExportSessionData();
+
+	// For each set of coplanar markers, add a new calibration row
+	if (ArucoCalibrationPoints.Num() > 0)
+	{
+		TSharedPtr<FLensDistortionArucoRowData> NewRow = MakeShared<FLensDistortionArucoRowData>();
+		NewRow->Index = LensDistortionTool->AdvanceSessionRowIndex();
+
+		// Save the aruco points associated with this row 
+		NewRow->ArucoPoints = ArucoCalibrationPoints;
+
+		// Get the current pose of the selected CineCamera actor
+		NewRow->CameraPose = FTransform::Identity;
+		if (const ACameraActor* Camera = StepsController->GetCamera())
+		{
+			if (const UCameraComponent* CameraComponent = Camera->GetCameraComponent())
+			{
+				NewRow->CameraPose = CameraComponent->GetComponentToWorld();
+			}
+		}
+
+		// Save an image view of the captured frame
+		FImageView ImageView = FImageView(Pixels.GetData(), Size.X, Size.Y, ERawImageFormat::BGRA8);
+		ImageView.CopyTo(NewRow->MediaImage);
+
+		CalibrationRows.Add(NewRow);
+
+		// Export the data for this row to a .json file on disk
+		ExportRow(NewRow);
+	}
+
+	// Notify the ListView of the new data
+	if (CalibrationRowListWidget)
+	{
+		CalibrationRowListWidget->RequestListRefresh();
+	}
+
+	return true;
+}
+
+void UCameraLensDistortionAlgoAruco::ClearCalibrationRows()
+{
+	CalibrationRows.Empty();
+
+	if (CalibrationRowListWidget)
+	{
+		CalibrationRowListWidget->RequestListRefresh();
+	}
+	
+	// End the current calibration session (a new one will begin the next time a new row is added)
+	if (ULensDistortionTool* LenDistortionTool = WeakTool.Get())
+	{
+		LenDistortionTool->EndCalibrationSession();
+	}
+}
+
+void UCameraLensDistortionAlgoAruco::SetCalibrator(AActor* InCalibrator)
+{
+	WeakCalibrator = InCalibrator;
+
+	if (!InCalibrator)
+	{
+		return;
+	}
+
+	// Find all of the calibration components of the input calibrator 
+	constexpr uint32 NumInlineAllocations = 32;
+	TArray<UCalibrationPointComponent*, TInlineAllocator<NumInlineAllocations>> CalibrationPoints;
+	InCalibrator->GetComponents(CalibrationPoints);
+
+	ActiveCalibratorComponents.Empty();
+
+	// When a new calibrator is selected, all of its components will default to active
+	for (UCalibrationPointComponent* CalibrationPoint : CalibrationPoints)
+	{
+		// Only add calibration point components that have an attached scene component
+		if (CalibrationPoint && CalibrationPoint->GetAttachParent())
+		{
+			ActiveCalibratorComponents.Add(CalibrationPoint);
+		}
+	}
+}
+
+void UCameraLensDistortionAlgoAruco::ExportSessionData()
+{
+	using namespace UE::CameraCalibration::Private;
+
+	if (ULensDistortionTool* LensDistortionTool = WeakTool.Get())
+	{
+		// Add all data to a json object that is needed to run this algorithm that is NOT part of a specific row
+		TSharedPtr<FJsonObject> JsonSessionData = MakeShared<FJsonObject>();
+
+		JsonSessionData->SetNumberField(LensDistortionArucoExportFields::Version, DATASET_VERSION);
+
+		// Export the session data to a .json file
+		LensDistortionTool->ExportSessionData(JsonSessionData.ToSharedRef());
+	}
+}
+
+void UCameraLensDistortionAlgoAruco::ExportRow(TSharedPtr<FLensDistortionArucoRowData> Row)
+{
+	if (ULensDistortionTool* LensDistortionTool = WeakTool.Get())
+	{
+		if (const TSharedPtr<FJsonObject>& RowObject = FJsonObjectConverter::UStructToJsonObject<FLensDistortionArucoRowData>(Row.ToSharedRef().Get()))
+		{
+			// Export the row data to a .json file
+			LensDistortionTool->ExportCalibrationRow(Row->Index, RowObject.ToSharedRef(), Row->MediaImage);
+		}
+	}
+}
+
+bool UCameraLensDistortionAlgoAruco::HasCalibrationData() const
+{
+	return (CalibrationRows.Num() > 0);
+}
+
+void UCameraLensDistortionAlgoAruco::PreImportCalibrationData()
+{
+	// Clear the current set of rows before importing new ones
+	CalibrationRows.Empty();
+}
+
+int32 UCameraLensDistortionAlgoAruco::ImportCalibrationRow(const TSharedRef<FJsonObject>& CalibrationRowObject, const FImage& RowImage)
+{
+	// Create a new row to populate with data from the Json object
+	TSharedPtr<FLensDistortionArucoRowData> NewRow = MakeShared<FLensDistortionArucoRowData>();
+
+	// We enforce strict mode to ensure that every field in the UStruct of row data is present in the imported json.
+	// If any fields are missing, it is likely the row will be invalid, which will lead to errors in the calibration.
+	constexpr int64 CheckFlags = 0;
+	constexpr int64 SkipFlags = 0;
+	constexpr bool bStrictMode = true;
+	if (FJsonObjectConverter::JsonObjectToUStruct<FLensDistortionArucoRowData>(CalibrationRowObject, NewRow.Get(), CheckFlags, SkipFlags, bStrictMode))
+	{
+		if (!RowImage.RawData.IsEmpty())
+		{
+			NewRow->MediaImage = RowImage;
+		}
+
+		CalibrationRows.Add(NewRow);
+	}
+	else
+	{
+		UE_LOG(LogCameraCalibrationEditor, Warning, TEXT("LensDistortionAruco algo failed to import calibration row because at least one field could not be deserialized from the json file."));
+	}
+
+	return NewRow->Index;
+}
+
+void UCameraLensDistortionAlgoAruco::PostImportCalibrationData()
+{
+	// Sort imported calibration rows by row index
+	CalibrationRows.Sort([](const TSharedPtr<FLensDistortionArucoRowData>& LHS, const TSharedPtr<FLensDistortionArucoRowData>& RHS) { return LHS->Index < RHS->Index; });
+
+	// Notify the ListView of the new data
+	if (CalibrationRowListWidget)
+	{
+		CalibrationRowListWidget->RequestListRefresh();
+	}
+}
+
+TSharedRef<SWidget> UCameraLensDistortionAlgoAruco::BuildUI()
+{
+	return SNew(SVerticalBox)
+
+		+ SVerticalBox::Slot() // Calibrator picker
+		.VAlign(EVerticalAlignment::VAlign_Top)
+		.AutoHeight()
+		.MaxHeight(FCameraCalibrationWidgetHelpers::DefaultRowHeight)
+		[FCameraCalibrationWidgetHelpers::BuildLabelWidgetPair(LOCTEXT("Aruco", "Aruco"), BuildCalibratorPickerWidget())]
+
+		+ SVerticalBox::Slot() // Calibrator component picker
+		.VAlign(EVerticalAlignment::VAlign_Top)
+		.AutoHeight()
+		.MaxHeight(FCameraCalibrationWidgetHelpers::DefaultRowHeight)
+		[FCameraCalibrationWidgetHelpers::BuildLabelWidgetPair(LOCTEXT("CalibratorComponents", "Calibrator Component(s)"), BuildCalibrationComponentPickerWidget())]
+
+		+ SVerticalBox::Slot() // Solver Settings Title
+		.Padding(0, 5)
+		.AutoHeight()
+		.VAlign(EVerticalAlignment::VAlign_Center)
+		[
+			SNew(SBox) // Constrain the height
+			.MinDesiredHeight(FCameraCalibrationWidgetHelpers::DefaultRowHeight)
+			.MaxDesiredHeight(FCameraCalibrationWidgetHelpers::DefaultRowHeight)
+			[
+				SNew(SBorder) // Background color of title
+				.BorderImage(FAppStyle::GetBrush("DetailsView.CategoryTop"))
+				.BorderBackgroundColor(FLinearColor::White)
+				.VAlign(EVerticalAlignment::VAlign_Center)
+				[
+					SNew(SOverlay)
+					+ SOverlay::Slot() // Used to add left padding to the title
+					.Padding(5, 0, 0, 0)
+					[
+						SNew(STextBlock) // Title text
+						.Text(LOCTEXT("SolverSettings", "Solver Settings"))
+						.TransformPolicy(ETextTransformPolicy::ToUpper)
+						.Font(FAppStyle::Get().GetFontStyle(TEXT("PropertyWindow.BoldFont")))
+						.TextStyle(FAppStyle::Get(), "DetailsView.CategoryTextStyle")
+					]
+				]
+			]
+		]
+
+		+ SVerticalBox::Slot() // Focal Length Estimate
+		.VAlign(EVerticalAlignment::VAlign_Top)
+		.AutoHeight()
+		.MaxHeight(FCameraCalibrationWidgetHelpers::DefaultRowHeight)
+		[FCameraCalibrationWidgetHelpers::BuildLabelWidgetPair(LOCTEXT("FocalLengthEstimate", "Focal Length Estimate"), BuildFocalLengthEstimateWidget())]
+
+		+ SVerticalBox::Slot() // Fix Image Center
+		.VAlign(EVerticalAlignment::VAlign_Top)
+		.AutoHeight()
+		.MaxHeight(FCameraCalibrationWidgetHelpers::DefaultRowHeight)
+		[FCameraCalibrationWidgetHelpers::BuildLabelWidgetPair(LOCTEXT("FixImageCenter", "Fix Image Center"), BuildFixImageCenterWidget())]
+
+		+ SVerticalBox::Slot() // Calibration Rows
+		.AutoHeight()
+		.MaxHeight(12 * FCameraCalibrationWidgetHelpers::DefaultRowHeight)
+		[BuildCalibrationRowListWidget()]
+
+		+ SVerticalBox::Slot() // Action buttons (e.g. Clear)
+		.HAlign(EHorizontalAlignment::HAlign_Center)
+		.AutoHeight()
+		.Padding(0, 20)
+		[BuildCalibrationActionButtons()];
+}
+
+TSharedRef<SWidget> UCameraLensDistortionAlgoAruco::BuildCalibratorPickerWidget()
+{
+	return SNew(SFilterableActorPicker)
+		.OnSetObject_UObject(this, &UCameraLensDistortionAlgoAruco::OnCalibratorSelected)
+		.OnShouldFilterAsset_UObject(this, &UCameraLensDistortionAlgoAruco::DoesAssetHaveCalibrationComponent)
+		.ActorAssetData_UObject(this, &UCameraLensDistortionAlgoAruco::GetCalibratorAssetData);
+}
+
+TSharedRef<SWidget> UCameraLensDistortionAlgoAruco::BuildCalibrationComponentPickerWidget()
+{
+	return SNew(SHorizontalBox)
+
+		+ SHorizontalBox::Slot()
+		.AutoWidth()
+		.Padding(FMargin(0.0f, 0.0f, 0.0f, 0.0f))
+		[
+			SNew(SComboButton)
+			.OnGetMenuContent_UObject(this, &UCameraLensDistortionAlgoAruco::BuildCalibrationComponentMenu)
+			.ContentPadding(FMargin(4.0, 2.0))
+			.ButtonContent()
+			[
+				SNew(STextBlock).Text_UObject(this, &UCameraLensDistortionAlgoAruco::GetCalibrationComponentMenuText)
+			]
+		];
+}
+
+TSharedRef<SWidget> UCameraLensDistortionAlgoAruco::BuildCalibrationComponentMenu()
+{
+	// Generate menu
+	FMenuBuilder MenuBuilder(true, nullptr);
+	MenuBuilder.BeginSection("CalibrationComponents", LOCTEXT("CalibrationComponents", "Calibration Point Components"));
+
+	if (AActor* Calibrator = WeakCalibrator.Get())
+	{
+		constexpr uint32 NumInlineAllocations = 32;
+		TArray<UCalibrationPointComponent*, TInlineAllocator<NumInlineAllocations>> CalibrationPointComponents;
+		Calibrator->GetComponents(CalibrationPointComponents);
+
+		for (UCalibrationPointComponent* CalibratorComponent : CalibrationPointComponents)
+		{
+			MenuBuilder.AddMenuEntry(
+				FText::Format(LOCTEXT("ComponentLabel", "{0}"), FText::FromString(CalibratorComponent->GetName())),
+				FText::Format(LOCTEXT("ComponentTooltip", "{0}"), FText::FromString(CalibratorComponent->GetName())),
+				FSlateIcon(),
+				FUIAction(
+					FExecuteAction::CreateUObject(this, &UCameraLensDistortionAlgoAruco::OnCalibrationComponentSelected, CalibratorComponent),
+					FCanExecuteAction(),
+					FIsActionChecked::CreateUObject(this, &UCameraLensDistortionAlgoAruco::IsCalibrationComponentSelected, CalibratorComponent)
+				),
+				NAME_None,
+				EUserInterfaceActionType::ToggleButton
+			);
+		}
+	}
+	else
+	{
+		MenuBuilder.AddMenuEntry(
+			FText::FromName(NAME_None),
+			FText::FromName(NAME_None),
+			FSlateIcon(),
+			FUIAction(
+				FExecuteAction(),
+				FCanExecuteAction(),
+				FIsActionChecked()
+			),
+			NAME_None,
+			EUserInterfaceActionType::None
+		);
+	}
+	MenuBuilder.EndSection();
+
+	return MenuBuilder.MakeWidget();
+}
+
+TSharedRef<SWidget> UCameraLensDistortionAlgoAruco::BuildFocalLengthEstimateWidget()
+{
+	return SNew(SHorizontalBox)
+
+		+ SHorizontalBox::Slot()
+		.FillWidth(0.8)
+		.Padding(0, 0, 10, 0)
+		[
+			SNew(SNumericEntryBox<float>)
+			.Value_UObject(this, &UCameraLensDistortionAlgoAruco::GetFocalLengthEstimate)
+			.OnValueChanged_UObject(this, &UCameraLensDistortionAlgoAruco::SetFocalLengthEstimate)
+		]
+
+		+ SHorizontalBox::Slot()
+		.FillWidth(0.2)
+		[
+			SNew(SCheckBox)
+			.Padding(FMargin(5, 0, 15, 0))
+			.IsChecked_UObject(this, &UCameraLensDistortionAlgoAruco::IsFixFocalLengthChecked)
+			.OnCheckStateChanged_UObject(this, &UCameraLensDistortionAlgoAruco::OnFixFocalLengthCheckStateChanged)
+			[
+				SNew(STextBlock).Text(LOCTEXT("FixText", "Fix?"))
+			]
+		];
+}
+
+TSharedRef<SWidget> UCameraLensDistortionAlgoAruco::BuildFixImageCenterWidget()
+{
+	return SNew(SCheckBox)
+		.IsChecked_UObject(this, &UCameraLensDistortionAlgoAruco::IsFixImageCenterChecked)
+		.OnCheckStateChanged_UObject(this, &UCameraLensDistortionAlgoAruco::OnFixImageCenterCheckStateChanged);
+}
+
+TSharedRef<SWidget> UCameraLensDistortionAlgoAruco::BuildCalibrationActionButtons()
+{
+	return SNew(SHorizontalBox)
+
+		+ SHorizontalBox::Slot() // Button to clear all rows
+		.AutoWidth()
+		[
+			SNew(SButton)
+			.Text(LOCTEXT("ClearAll", "Clear All"))
+			.HAlign(HAlign_Center)
+			.VAlign(VAlign_Center)
+			.OnClicked_UObject(this, &UCameraLensDistortionAlgoAruco::OnClearCalibrationRowsClicked)
+		];
+}
+
+TSharedRef<SWidget> UCameraLensDistortionAlgoAruco::BuildCalibrationRowListWidget()
+{
+	CalibrationRowListWidget = SNew(SListView<TSharedPtr<FLensDistortionArucoRowData>>)
+		.ListItemsSource(&CalibrationRows)
+		.SelectionMode(ESelectionMode::Multi)
+		.OnGenerateRow_Lambda([](TSharedPtr<FLensDistortionArucoRowData> InItem, const TSharedRef<STableViewBase>& OwnerTable) -> TSharedRef<ITableRow>
+		{
+			return SNew(CameraLensDistortionAlgoAruco::SCalibrationRowGenerator, OwnerTable).CalibrationRowData(InItem);
+		})
+		.OnMouseButtonDoubleClick_Lambda([](TSharedPtr<FLensDistortionArucoRowData> Row)
+		{
+			// Open a dialog showing the media image of the selected row with debug aruco markers drawn on it
+			UTexture2D* DebugImage = UE::CameraCalibration::Private::CreateDebugTexture(Row);
+			UE::CameraCalibration::Private::DisplayDebugImageDialog(DebugImage);
+		})
+		.OnKeyDownHandler_UObject(this, &UCameraLensDistortionAlgoAruco::OnCalibrationRowListKeyPressed)
+		.HeaderRow
+		(
+			SNew(SHeaderRow)
+
+			+ SHeaderRow::Column("Index")
+			.DefaultLabel(LOCTEXT("Index", "Index"))
+			.FillWidth(0.1f)
+
+			+ SHeaderRow::Column("NumMarkers")
+			.DefaultLabel(LOCTEXT("NumMarkers", "Num Markers"))
+			.FillWidth(0.2f)
+
+			+ SHeaderRow::Column("Image")
+			.DefaultLabel(LOCTEXT("Image", "Image"))
+			.FillWidth(0.7f)
+		);
+
+	return CalibrationRowListWidget.ToSharedRef();
+}
+
+void UCameraLensDistortionAlgoAruco::OnCalibratorSelected(const FAssetData& AssetData)
+{
+	if (AssetData.IsValid())
+	{
+		SetCalibrator(Cast<AActor>(AssetData.GetAsset()));
+	}
+}
+
+bool UCameraLensDistortionAlgoAruco::DoesAssetHaveCalibrationComponent(const FAssetData& AssetData) const
+{
+	if (const AActor* Actor = Cast<AActor>(AssetData.GetAsset()))
+	{
+		constexpr uint32 NumInlineAllocations = 32;
+		TArray<UCalibrationPointComponent*, TInlineAllocator<NumInlineAllocations>> CalibrationPoints;
+		Actor->GetComponents(CalibrationPoints);
+
+		return (CalibrationPoints.Num() > 0);
+	}
+
+	return false;
+}
+
+FAssetData UCameraLensDistortionAlgoAruco::GetCalibratorAssetData() const
+{
+	return FAssetData(WeakCalibrator.Get(), true);
+}
+
+FText UCameraLensDistortionAlgoAruco::GetCalibrationComponentMenuText() const 
+{
+	if (ActiveCalibratorComponents.Num() > 1)
+	{
+		return LOCTEXT("MultipleCalibrationComponents", "Multiple Values");
+	}
+	else if (ActiveCalibratorComponents.Num() == 1)
+	{
+		if (UCalibrationPointComponent* CalibratorComponent = ActiveCalibratorComponents[0].Get())
+		{
+			return FText::FromString(CalibratorComponent->GetName());
+		}
+	}
+
+	return LOCTEXT("NoCalibrationComponents", "None");
+}
+
+void UCameraLensDistortionAlgoAruco::OnCalibrationComponentSelected(UCalibrationPointComponent* SelectedComponent)
+{
+	if (IsCalibrationComponentSelected(SelectedComponent))
+	{
+		ActiveCalibratorComponents.Remove(SelectedComponent);
+	}
+	else
+	{
+		ActiveCalibratorComponents.Add(SelectedComponent);
+	}
+}
+
+bool UCameraLensDistortionAlgoAruco::IsCalibrationComponentSelected(UCalibrationPointComponent* SelectedComponent) const
+{
+	return ActiveCalibratorComponents.Contains(SelectedComponent);
+}
+
+TOptional<float> UCameraLensDistortionAlgoAruco::GetFocalLengthEstimate() const
+{
+	return FocalLengthEstimate;
+}
+
+void UCameraLensDistortionAlgoAruco::SetFocalLengthEstimate(float NewValue)
+{
+	FocalLengthEstimate = NewValue;
+}
+
+ECheckBoxState UCameraLensDistortionAlgoAruco::IsFixFocalLengthChecked() const
+{
+	return bFixFocalLength ? ECheckBoxState::Checked : ECheckBoxState::Unchecked;
+}
+
+void UCameraLensDistortionAlgoAruco::OnFixFocalLengthCheckStateChanged(ECheckBoxState NewState)
+{
+	bFixFocalLength = (NewState == ECheckBoxState::Checked);
+}
+
+ECheckBoxState UCameraLensDistortionAlgoAruco::IsFixImageCenterChecked() const
+{
+	return bFixImageCenter ? ECheckBoxState::Checked : ECheckBoxState::Unchecked;
+}
+
+void UCameraLensDistortionAlgoAruco::OnFixImageCenterCheckStateChanged(ECheckBoxState NewState)
+{
+	bFixImageCenter = (NewState == ECheckBoxState::Checked);
+}
+
+FReply UCameraLensDistortionAlgoAruco::OnClearCalibrationRowsClicked()
+{
+	ClearCalibrationRows();
+	return FReply::Handled();
+}
+
+FReply UCameraLensDistortionAlgoAruco::OnCalibrationRowListKeyPressed(const FGeometry& Geometry, const FKeyEvent& KeyEvent)
+{
+	if (!CalibrationRowListWidget)
+	{
+		return FReply::Unhandled();
+	}
+
+	if (KeyEvent.GetKey() == EKeys::Delete)
+	{
+		// Delete selected items
+		const TArray<TSharedPtr<FLensDistortionArucoRowData>> SelectedItems = CalibrationRowListWidget->GetSelectedItems();
+
+		for (const TSharedPtr<FLensDistortionArucoRowData>& SelectedItem : SelectedItems)
+		{
+			CalibrationRows.Remove(SelectedItem);
+
+			// Delete the previously exported .json file for the row that is being deleted
+			if (ULensDistortionTool* LensDistortionTool = WeakTool.Get())
+			{
+				LensDistortionTool->DeleteExportedRow(SelectedItem->Index);
+			}
+		}
+
+		CalibrationRowListWidget->RequestListRefresh();
+		return FReply::Handled();
+	}
+	else if (KeyEvent.GetModifierKeys().IsControlDown() && (KeyEvent.GetKey() == EKeys::A))
+	{
+		// Select all items
+		CalibrationRowListWidget->SetItemSelection(CalibrationRows, true);
+		return FReply::Handled();
+	}
+	else if (KeyEvent.GetKey() == EKeys::Escape)
+	{
+		// De-select all items
+		CalibrationRowListWidget->ClearSelection();
+		return FReply::Handled();
+	}
+
+	return FReply::Unhandled();
+}
+
+#undef LOCTEXT_NAMESPACE
diff --git a/Engine/Plugins/VirtualProduction/CameraCalibration/Source/CameraCalibrationEditor/Private/Calibrators/CameraLensDistortionAlgoAruco.h b/Engine/Plugins/VirtualProduction/CameraCalibration/Source/CameraCalibrationEditor/Private/Calibrators/CameraLensDistortionAlgoAruco.h
new file mode 100644
--- /dev/null
+++ b/Engine/Plugins/VirtualProduction/CameraCalibration/Source/CameraCalibrationEditor/Private/Calibrators/CameraLensDistortionAlgoAruco.h
@@ -0,0 +1,174 @@
+// Copyright Epic Games, Inc. All Rights Reserved.
+
+#pragma once
+
+#include "CameraLensDistortionAlgo.h"
+#include "ImageCore.h"
+
+#include "CameraLensDistortionAlgoAruco.generated.h"
+
+class UCalibrationPointComponent;
+
+template <typename ItemType>
+class SListView;
+
+struct FArucoCalibrationPoint;
+
+/** Holds information of the calibration row */
+USTRUCT()
+struct FLensDistortionArucoRowData
+{
+	GENERATED_BODY()
+
+	/** Index to display in list */
+	UPROPERTY()
+	int32 Index = -1;
+
+	/** Array of aruco points (with their 3D and 2D corner locations) */
+	UPROPERTY()
+	TArray<FArucoCalibrationPoint> ArucoPoints;
+
+	/** Pose of the tracked camera when the aruco points were detected */
+	UPROPERTY()
+	FTransform CameraPose = FTransform::Identity;
+
+	/** Cached image of the media texture that contained the aruco markers */
+	FImage MediaImage;
+};
+
+/** 
+ * Implements a lens distortion calibration algorithm. It requires a Aruco pattern
+ */
+UCLASS()
+class UCameraLensDistortionAlgoAruco : public UCameraLensDistortionAlgo
+{
+	GENERATED_BODY()
+	
+public:
+
+	//~ Begin CameraLensDistortionAlgo
+	virtual void Initialize(ULensDistortionTool* InTool) override;
+	virtual void Shutdown() override;
+	virtual bool SupportsModel(const TSubclassOf<ULensModel>& LensModel) const override;
+	virtual void Tick(float DeltaTime) override;
+	virtual bool OnViewportClicked(const FGeometry& MyGeometry, const FPointerEvent& MouseEvent) override;
+	virtual TSharedRef<SWidget> BuildUI() override;
+	virtual FName FriendlyName() const override { return TEXT("Lens Distortion Aruco"); };
+	virtual FName ShortName() const override { return TEXT("Aruco"); };
+	virtual void OnDistortionSavedToLens() override;
+	virtual FDistortionCalibrationTask BeginCalibration(FText& OutErrorMessage) override;
+	virtual bool SupportsAsyncCalibration() override { return true; };
+	virtual bool HasCalibrationData() const override;
+	virtual void PreImportCalibrationData() override;
+	virtual int32 ImportCalibrationRow(const TSharedRef<FJsonObject>& CalibrationRowObject, const FImage& RowImage) override;
+	virtual void PostImportCalibrationData() override;
+	//~ End CameraLensDistortionAlgo
+
+private:
+	/** Add a new calibration row from media texture and camera data */
+	bool AddCalibrationRow(FText& OutErrorMessage);
+
+	/** Clears the list of calibration sample rows */
+	void ClearCalibrationRows();
+
+	/** Sets the calibrator object to be used. Updates the selection picker. */
+	void SetCalibrator(AActor* InCalibrator);
+
+	/** Export global session data to a .json file */
+	void ExportSessionData();
+
+	/** Export the row data to a json file */
+	void ExportRow(TSharedPtr<FLensDistortionArucoRowData> Row);
+
+private:
+	/** Builds the UI of the calibration device picker */
+	TSharedRef<SWidget> BuildCalibratorPickerWidget();
+
+	/** Builds the UI of the fix focal length width */
+	TSharedRef<SWidget> BuildFocalLengthEstimateWidget();
+
+	/** Builds the UI of the fix image center widget */
+	TSharedRef<SWidget> BuildFixImageCenterWidget();
+
+	/** Builds the UI of the calibration points table */
+	TSharedRef<SWidget> BuildCalibrationRowListWidget();
+
+	/** Builds the UI for the action buttons (RemoveLast, ClearAll) */
+	TSharedRef<SWidget> BuildCalibrationActionButtons();
+
+	/** Builds the UI of the calibration component picker */
+	TSharedRef<SWidget> BuildCalibrationComponentPickerWidget();
+
+	/** Builds the UI of the calibration component picker */
+	TSharedRef<SWidget> BuildCalibrationComponentMenu();
+
+	/** Gets the actor from the input asset data and caches it as the current calibrator */
+	void OnCalibratorSelected(const FAssetData& AssetData);
+
+	/** Returns true if the actor represented by the input asset data has any calibration point components. Used to filter calibrator actors in the picker. */
+	bool DoesAssetHaveCalibrationComponent(const FAssetData& AssetData) const;
+
+	/** Get an FAssetData representation of the current calibrator actor */
+	FAssetData GetCalibratorAssetData() const;
+
+	/** Get the text to display in the calibration component combobox. */
+	FText GetCalibrationComponentMenuText() const;
+
+	/** Add or remove the selected component from the set of active components and refresh the calibrator point combobox */
+	void OnCalibrationComponentSelected(UCalibrationPointComponent* SelectedComponent);
+
+	/** Returns true if the input component is in the set of active calibration components */
+	bool IsCalibrationComponentSelected(UCalibrationPointComponent* SelectedComponent) const;
+
+	/** Get the current value of the focal length estimate */
+	TOptional<float> GetFocalLengthEstimate() const;
+
+	/** Set the current value of the focal length estimate */
+	void SetFocalLengthEstimate(float NewValue);
+
+	/** Returns the state of the fix focal length checkbox */
+	ECheckBoxState IsFixFocalLengthChecked() const;
+
+	/** Changes the state of the fix focal length checkbox */
+	void OnFixFocalLengthCheckStateChanged(ECheckBoxState NewState);
+
+	/** Returns the state of the fix image center checkbox */
+	ECheckBoxState IsFixImageCenterChecked() const;
+
+	/** Changes the state of the fix image center checkbox */
+	void OnFixImageCenterCheckStateChanged(ECheckBoxState NewState);
+
+	/** Event triggered when the Clear button is pressed. Clears all of the captured calibration rows */
+	FReply OnClearCalibrationRowsClicked();
+
+	/** Event triggered when a key is pressed on the calibration list view */
+	FReply OnCalibrationRowListKeyPressed(const FGeometry& Geometry, const FKeyEvent& KeyEvent);
+
+private:
+	/** Version of the current calibration dataset */
+	static const int DATASET_VERSION;
+
+	/** The lens distortion tool controller */
+	TWeakObjectPtr<ULensDistortionTool> WeakTool;
+
+	/** The currently selected calibrator object. It is expected to contain one or more UCalibrationPointComponent in it */
+	TWeakObjectPtr<AActor> WeakCalibrator;
+
+	/** Rows source for the CalibrationRowListWidget */
+	TArray<TSharedPtr<FLensDistortionArucoRowData>> CalibrationRows;
+
+	/** Displays the list of calibration points that will be used to calculate the lens distortion */
+	TSharedPtr<SListView<TSharedPtr<FLensDistortionArucoRowData>>> CalibrationRowListWidget;
+
+	/** If true, the solver will not solve for focal length, and will just use the input value */
+	bool bFixFocalLength = false;
+
+	/** If true, the solver will not solve for image center, and will just use the input value */
+	bool bFixImageCenter = false;
+
+	/** Estimate for the focal length to provide to the solver. If not set, user will be warned to set before continuing calibration */
+	TOptional<float> FocalLengthEstimate;
+
+	/** Container for the set of calibrator components selected in the component combobox */
+	TArray<TWeakObjectPtr<UCalibrationPointComponent>> ActiveCalibratorComponents;
+};
diff --git a/Engine/Plugins/VirtualProduction/CameraCalibration/Source/CameraCalibrationEditor/Private/Calibrators/CameraLensDistortionAlgoCheckerboard.cpp b/Engine/Plugins/VirtualProduction/CameraCalibration/Source/CameraCalibrationEditor/Private/Calibrators/CameraLensDistortionAlgoCheckerboard.cpp
new file mode 100644
--- /dev/null
+++ b/Engine/Plugins/VirtualProduction/CameraCalibration/Source/CameraCalibrationEditor/Private/Calibrators/CameraLensDistortionAlgoCheckerboard.cpp
@@ -0,0 +1,1313 @@
+// Copyright Epic Games, Inc. All Rights Reserved.
+
+#include "CameraLensDistortionAlgoCheckerboard.h"
+
+#include "AssetEditor/CameraCalibrationStepsController.h"
+#include "AssetEditor/LensDistortionTool.h"
+#include "AssetEditor/SSimulcamViewport.h"
+#include "CalibrationPointComponent.h"
+#include "Camera/CameraActor.h"
+#include "CameraCalibrationCheckerboard.h"
+#include "CameraCalibrationEditorLog.h"
+#include "CameraCalibrationSolver.h"
+#include "CameraCalibrationUtilsPrivate.h"
+#include "Dom/JsonObject.h"
+#include "EditorFontGlyphs.h"
+#include "Engine/Texture2D.h"
+#include "EngineUtils.h"
+#include "ImageUtils.h"
+#include "JsonObjectConverter.h"
+#include "LensComponent.h"
+#include "LensFile.h"
+#include "Materials/MaterialInstanceDynamic.h"
+#include "Misc/MessageDialog.h"
+#include "Models/AnamorphicLensModel.h"
+#include "Models/SphericalLensModel.h"
+#include "OpenCVHelper.h"
+#include "PropertyCustomizationHelpers.h"
+#include "TextureResource.h"
+#include "UI/CameraCalibrationWidgetHelpers.h"
+#include "Widgets/Input/SButton.h"
+#include "Widgets/Input/SCheckBox.h"
+#include "Widgets/Input/SNumericEntryBox.h"
+#include "Widgets/Views/SListView.h"
+#include "Widgets/Views/STableRow.h"
+
+
+#define LOCTEXT_NAMESPACE "CameraLensDistortionAlgoCheckerboard"
+
+#if WITH_EDITOR
+static TAutoConsoleVariable<bool> CVarUseIntrinsicsGuess(TEXT("LensDistortionCheckerboard.UseIntrinsicsGuess"), true, TEXT("If true, the solver initializes the camera intrinsics to a user-provided estimate. Otherwise, the solver will compute the initial values."));
+static TAutoConsoleVariable<bool> CVarFixExtrinsics(TEXT("LensDistortionCheckerboard.FixExtrinsics"), false, TEXT("If true, the solver will fix the camera extrinsics to the user-provided camera poses"));
+static TAutoConsoleVariable<bool> CVarFixDistortion(TEXT("LensDistortionCheckerboard.FixDistortion"), false, TEXT("If true, the solver will not optimize distortion, and will use the input distortion values if any are given, or assume zero distortion otherwise."));
+static TAutoConsoleVariable<bool> CVarUseExtrinsicsGuess(TEXT("LensDistortionCheckerboard.UseExtrinsicsGuess"), false, TEXT("If true, the actual checkerboard and camera poses will be used when running the solver"));
+static TAutoConsoleVariable<bool> CVarUseTrackedCalibrator(TEXT("LensDistortionCheckerboard.UseTrackedCalibrator"), false, TEXT("If true, the 3D points of the checkerboard will come from the tracked pose of the object, otherwise, dummy points will be used."));
+static TAutoConsoleVariable<bool> CVarSolveTargetOffset(TEXT("LensDistortionCheckerboard.SolveTargetOffset"), false, TEXT("If true, the solver will calibrate for an offset from the calibrator tracking origin that further minimizes reprojection error."));
+#endif
+
+const int UCameraLensDistortionAlgoCheckerboard::DATASET_VERSION = 1;
+
+// String constants for import/export
+namespace UE::CameraCalibration::Private::LensDistortionCheckerboardExportFields
+{
+	static const FString Version(TEXT("Version"));
+}
+
+namespace UE::CameraCalibration::Private
+{
+	void AdjustCorners(TArray<FVector2f>& Corners, FIntPoint TextureSize, FIntPoint ImageSize)
+	{
+		// It is possible that the size of the debug texture where we want to draw the checkerboard corners is different than the size of the image where the checkerboard corners were found.
+		// Before drawing, all of the corner positions should be shifted to account for the difference in size.
+		const FVector2f TopLeftCorner = (TextureSize - ImageSize) / 2.0f;
+
+		for (FVector2f& Corner : Corners)
+		{
+			Corner += TopLeftCorner;
+		}
+	}
+}
+
+namespace CameraLensDistortionAlgoCheckerboard
+{
+	class SCalibrationRowGenerator
+		: public SMultiColumnTableRow<TSharedPtr<FLensDistortionCheckerboardRowData>>
+	{
+
+	public:
+		SLATE_BEGIN_ARGS(SCalibrationRowGenerator) {}
+
+		/** The list item for this row */
+		SLATE_ARGUMENT(TSharedPtr<FLensDistortionCheckerboardRowData>, CalibrationRowData)
+
+		SLATE_END_ARGS()
+
+		void Construct(const FArguments& Args, const TSharedRef<STableViewBase>& OwnerTableView)
+		{
+			CalibrationRowData = Args._CalibrationRowData;
+
+			SMultiColumnTableRow<TSharedPtr<FLensDistortionCheckerboardRowData>>::Construct(
+				FSuperRowType::FArguments()
+				.Padding(1.0f),
+				OwnerTableView
+			);
+		}
+
+		//~Begin SMultiColumnTableRow
+		virtual TSharedRef<SWidget> GenerateWidgetForColumn(const FName& ColumnName) override
+		{
+			if (ColumnName == TEXT("Index"))
+			{
+				return SNew(SVerticalBox)
+					+ SVerticalBox::Slot()
+					.HAlign(EHorizontalAlignment::HAlign_Center)
+					.VAlign(EVerticalAlignment::VAlign_Center)
+					[SNew(STextBlock).Text(FText::AsNumber(CalibrationRowData->Index))];
+			}
+
+			if (ColumnName == TEXT("Image"))
+			{
+				if (CalibrationRowData->Thumbnail.IsValid())
+				{
+					check((CalibrationRowData->ImageHeight) > 0 && (CalibrationRowData->ImageWidth > 0));
+
+					const float AspectRatio = float(CalibrationRowData->ImageWidth) / float(CalibrationRowData->ImageHeight);
+
+					return SNew(SBox)
+						.MinAspectRatio(AspectRatio)
+						.MaxAspectRatio(AspectRatio)
+						.MinDesiredHeight(4 * FCameraCalibrationWidgetHelpers::DefaultRowHeight)
+						.MaxDesiredHeight(4 * FCameraCalibrationWidgetHelpers::DefaultRowHeight)
+						[CalibrationRowData->Thumbnail.ToSharedRef()];
+				}
+				else
+				{
+					const FString Text = FString::Printf(TEXT("Image Unavailable"));
+					return SNew(STextBlock).Text(FText::FromString(Text));
+				}
+			}
+
+			return SNullWidget::NullWidget;
+		}
+		//~End SMultiColumnTableRow
+
+
+	private:
+		TSharedPtr<FLensDistortionCheckerboardRowData> CalibrationRowData;
+	};
+
+};
+
+void UCameraLensDistortionAlgoCheckerboard::Initialize(ULensDistortionTool* InTool)
+{
+	Tool = InTool;
+
+	// Guess which calibrator to use.
+	SetCalibrator(FindFirstCalibrator());
+
+	// Initialize coverage texture
+	FCameraCalibrationStepsController* StepsController = Tool->GetCameraCalibrationStepsController();
+
+	if (!ensure(StepsController))
+	{
+		return;
+	}
+
+	const FIntPoint Size = StepsController->GetCompRenderResolution();
+
+	CoverageTexture = UTexture2D::CreateTransient(Size.X, Size.Y, EPixelFormat::PF_B8G8R8A8);
+	UE::CameraCalibration::Private::ClearTexture(CoverageTexture);
+
+	if (UMaterialInstanceDynamic* OverlayMID = Tool->GetOverlayMID())
+	{
+		OverlayMID->SetTextureParameterValue(FName(TEXT("CoverageTexture")), CoverageTexture);
+	}
+}
+
+void UCameraLensDistortionAlgoCheckerboard::Shutdown()
+{
+	Tool.Reset();
+	CalibrationRows.Reset();
+}
+
+bool UCameraLensDistortionAlgoCheckerboard::SupportsModel(const TSubclassOf<ULensModel>& LensModel) const
+{
+	return ((LensModel == USphericalLensModel::StaticClass()) || (LensModel == UAnamorphicLensModel::StaticClass()));
+}
+
+void UCameraLensDistortionAlgoCheckerboard::Tick(float DeltaTime)
+{
+	if (!ensure(Tool.IsValid()))
+	{
+		return;
+	}
+
+	FCameraCalibrationStepsController* StepsController = Tool->GetCameraCalibrationStepsController();
+
+	if (!ensure(StepsController))
+	{
+		return;
+	}
+
+	// If the resolution of the simulcam comp has changed, update the coverage texture to be the correct size
+	const FIntPoint Size = StepsController->GetCompRenderResolution();
+	if (CoverageTexture && ((CoverageTexture->GetSizeX() != Size.X) || (CoverageTexture->GetSizeY() != Size.Y)))
+	{
+		CoverageTexture = UTexture2D::CreateTransient(Size.X, Size.Y, EPixelFormat::PF_B8G8R8A8);
+		RefreshCoverage();
+	}
+
+	// If not paused, cache calibrator 3d point position
+	if (!StepsController->IsPaused())
+	{
+		// Cache camera data
+		do
+		{
+			LastCameraData.bIsValid = false;
+
+			const FLensFileEvaluationInputs LensFileEvalInputs = StepsController->GetLensFileEvaluationInputs();
+
+			// We require lens evaluation data.
+			if (LensFileEvalInputs.bIsValid)
+			{
+				LastCameraData.InputFocus = LensFileEvalInputs.Focus;
+				LastCameraData.InputZoom = LensFileEvalInputs.Zoom;
+			}
+
+			const ACameraActor* Camera = StepsController->GetCamera();
+
+			if (!Camera)
+			{
+				break;
+			}
+
+			const UCameraComponent* CameraComponent = Camera->GetCameraComponent();
+
+			if (!CameraComponent)
+			{
+				break;
+			}
+
+			LastCameraData.Pose = CameraComponent->GetComponentToWorld();
+
+			const ULensComponent* LensComponent = StepsController->FindLensComponent();
+			if (LensComponent)
+			{
+				LastCameraData.bWasNodalOffsetApplied = LensComponent->WasNodalOffsetAppliedThisTick();
+			}
+
+			LastCameraData.bIsValid = true;
+
+
+		} while (0);
+	}
+}
+
+bool UCameraLensDistortionAlgoCheckerboard::OnViewportClicked(const FGeometry& MyGeometry, const FPointerEvent& MouseEvent)
+{
+	// We only respond to left clicks
+	if (MouseEvent.GetEffectingButton() != EKeys::LeftMouseButton)
+	{
+		return false;
+	}
+
+	// Left Mouse button means to add a new calibration row
+	FText ErrorMessage;
+
+	if (!AddCalibrationRow(ErrorMessage))
+	{
+		FMessageDialog::Open(EAppMsgType::Ok, ErrorMessage);
+	}
+
+	// force play (whether we succeeded or not)
+	if (FCameraCalibrationStepsController* StepsController = GetStepsController())
+	{
+		StepsController->Play();
+	}
+
+	return true;
+}
+
+FCameraCalibrationStepsController* UCameraLensDistortionAlgoCheckerboard::GetStepsController() const
+{
+	if (!ensure(Tool.IsValid()))
+	{
+		return nullptr;
+	}
+
+	return Tool->GetCameraCalibrationStepsController();
+}
+
+bool UCameraLensDistortionAlgoCheckerboard::AddCalibrationRow(FText& OutErrorMessage)
+{
+	using namespace CameraLensDistortionAlgoCheckerboard;
+
+	if (!Tool.IsValid())
+	{
+		OutErrorMessage = LOCTEXT("InvalidTool", "Invalid Tool");
+		return false;
+	}
+
+	FCameraCalibrationStepsController* StepsController = Tool->GetCameraCalibrationStepsController();
+
+	if (!StepsController)
+	{
+		OutErrorMessage = LOCTEXT("InvalidStepsController", "Invalid StepsController");
+		return false;
+	}
+
+	if (!LastCameraData.bIsValid)
+	{
+		OutErrorMessage = LOCTEXT("InvalidLastCameraData", "Could not find a cached set of camera data (e.g. FIZ). Check the Lens Component to make sure it has valid evaluation inputs.");
+		return false;
+	}
+
+	if (!Calibrator.IsValid())
+	{
+		OutErrorMessage = LOCTEXT("InvalidCalibrator", "Please pick a calibrator actor in the given combo box that contains calibration points");
+		return false;
+	}
+
+	// Read pixels
+
+	TArray<FColor> Pixels;
+	FIntPoint Size;
+
+	if (!StepsController->ReadMediaPixels(Pixels, Size, OutErrorMessage, ESimulcamViewportPortion::CameraFeed))
+	{
+		return false;
+	}
+
+	// Export the latest session data
+	ExportSessionData();
+
+	// Create the row that we're going to add
+	TSharedPtr<FLensDistortionCheckerboardRowData> Row = MakeShared<FLensDistortionCheckerboardRowData>();
+
+	// Get the next row index for the current calibration session to assign to this new row
+	const uint32 RowIndex = Tool->AdvanceSessionRowIndex();
+
+	Row->Index = RowIndex;
+	Row->CameraData = LastCameraData;
+	Row->ImageWidth = Size.X;
+	Row->ImageHeight = Size.Y;
+
+	Row->NumCornerRows = Calibrator->NumCornerRows;
+	Row->NumCornerCols = Calibrator->NumCornerCols;
+	Row->SquareSideInCm = Calibrator->SquareSideLength;
+
+	const FIntPoint CheckerboardDimensions = FIntPoint(Calibrator->NumCornerCols, Calibrator->NumCornerRows);
+
+	const bool bCornersFound = FOpenCVHelper::IdentifyCheckerboard(Pixels, Size, CheckerboardDimensions, Row->Points2d);
+
+ 	if (!bCornersFound || Row->Points2d.IsEmpty())
+ 	{
+ 		OutErrorMessage = FText::FromString(FString::Printf(TEXT(
+ 			"Could not identify the expected checkerboard points of interest. "
+ 			"The expected checkerboard has %dx%d inner corners."),
+ 			Row->NumCornerCols, Row->NumCornerRows)
+ 		);
+ 
+ 		return false;
+ 	}
+
+	// Fill out the checkerboard's 3D points
+	const FVector LocationTL = Calibrator->TopLeft->GetComponentLocation();
+	const FVector LocationTR = Calibrator->TopRight->GetComponentLocation();
+	const FVector LocationBL = Calibrator->BottomLeft->GetComponentLocation();
+
+	const FVector RightVector = LocationTR - LocationTL;
+	const FVector DownVector = LocationBL - LocationTL;
+
+	const float HorizontalStep = (Row->NumCornerCols > 1) ? (1.0f / (Row->NumCornerCols - 1)) : 0.0f;
+	const float VerticalStep = (Row->NumCornerRows > 1) ? (1.0f / (Row->NumCornerRows - 1)) : 0.0f;
+
+	for (int32 RowIdx = 0; RowIdx < Row->NumCornerRows; ++RowIdx)
+	{
+		for (int32 ColIdx = 0; ColIdx < Row->NumCornerCols; ++ColIdx)
+		{
+			const FVector PointLocation = LocationTL + (RightVector * ColIdx * HorizontalStep) + (DownVector * RowIdx * VerticalStep);
+			Row->Points3d.Add(PointLocation);
+		}
+	}
+
+	// Save an image view of the captured frame (for exporting)
+	FImageView ImageView = FImageView(Pixels.GetData(), Size.X, Size.Y, ERawImageFormat::BGRA8);
+	Row->ImageView = ImageView;
+
+	// Generate a thumbnail image for display in the tool for this row
+	FImage RowThumbnail;
+	constexpr int32 ResolutionDivider = 4;
+	FImageCore::ResizeTo(ImageView, RowThumbnail, Size.X / ResolutionDivider, Size.Y / ResolutionDivider, ERawImageFormat::BGRA8, ImageView.GetGammaSpace());
+
+	UTexture2D* ThumbnailTexture = FImageUtils::CreateTexture2DFromImage(RowThumbnail);
+	if (ThumbnailTexture)
+	{
+		Row->Thumbnail = SNew(SSimulcamViewport, ThumbnailTexture).WithZoom(false).WithPan(false);
+	}
+
+	// Update the coverage overlay with the latest checkerboard corners
+	TArray<FVector2f> CameraFeedAdjustedCorners = Row->Points2d;
+	const FIntPoint CoverageTextureSize = FIntPoint(CoverageTexture->GetSizeX(), CoverageTexture->GetSizeY());
+	UE::CameraCalibration::Private::AdjustCorners(CameraFeedAdjustedCorners, CoverageTextureSize, Size);
+
+	FOpenCVHelper::DrawCheckerboardCorners(CameraFeedAdjustedCorners, CheckerboardDimensions, CoverageTexture);
+	StepsController->RefreshOverlay();
+
+	// Show the detection to the user
+	if (bShouldShowDetectionWindow)
+	{
+		UTexture2D* DebugTexture = UTexture2D::CreateTransient(Size.X, Size.Y, EPixelFormat::PF_B8G8R8A8);
+		UE::CameraCalibration::Private::SetTextureData(DebugTexture, Pixels);
+		FOpenCVHelper::DrawCheckerboardCorners(Row->Points2d, CheckerboardDimensions, DebugTexture);
+
+		FCameraCalibrationWidgetHelpers::DisplayTextureInWindowAlmostFullScreen(DebugTexture, LOCTEXT("CheckerboardDetection", "Checkerboard Detection"));
+	}
+
+	// Validate the new row, show a message if validation fails.
+	if (!ValidateNewRow(Row, OutErrorMessage))
+	{
+		return false;
+	}
+
+	// Add this data point
+	CalibrationRows.Add(Row);
+
+	// Export the data for this row to a .json file on disk
+	ExportRow(Row);
+
+	// Notify the ListView of the new data
+	if (CalibrationListView.IsValid())
+	{
+		CalibrationListView->RequestListRefresh();
+		CalibrationListView->RequestNavigateToItem(Row);
+	}
+
+	return true;
+}
+
+TSharedRef<SWidget> UCameraLensDistortionAlgoCheckerboard::BuildUI()
+{
+	return SNew(SVerticalBox)
+
+		+ SVerticalBox::Slot() // Calibrator picker
+		.VAlign(EVerticalAlignment::VAlign_Top)
+		.AutoHeight()
+		.MaxHeight(FCameraCalibrationWidgetHelpers::DefaultRowHeight)
+		[FCameraCalibrationWidgetHelpers::BuildLabelWidgetPair(LOCTEXT("Checkerboard", "Checkerboard"), BuildCalibrationDevicePickerWidget())]
+
+	+ SVerticalBox::Slot() // Show Overlay
+		.VAlign(EVerticalAlignment::VAlign_Top)
+		.AutoHeight()
+		.MaxHeight(FCameraCalibrationWidgetHelpers::DefaultRowHeight)
+		[FCameraCalibrationWidgetHelpers::BuildLabelWidgetPair(LOCTEXT("ShowOverlay", "Show Coverage Overlay"), BuildShowOverlayWidget())]
+
+	+ SVerticalBox::Slot() // Show Detection
+		.VAlign(EVerticalAlignment::VAlign_Top)
+		.AutoHeight()
+		.MaxHeight(FCameraCalibrationWidgetHelpers::DefaultRowHeight)
+		[FCameraCalibrationWidgetHelpers::BuildLabelWidgetPair(LOCTEXT("ShowDetection", "Show Detection"), BuildShowDetectionWidget())]
+
+	+ SVerticalBox::Slot() // Solver Settings Title
+		.Padding(0, 5)
+		.AutoHeight()
+		.VAlign(EVerticalAlignment::VAlign_Center)
+		[
+			SNew(SBox) // Constrain the height
+			.MinDesiredHeight(FCameraCalibrationWidgetHelpers::DefaultRowHeight)
+			.MaxDesiredHeight(FCameraCalibrationWidgetHelpers::DefaultRowHeight)
+			[
+				SNew(SBorder) // Background color of title
+				.BorderImage(FAppStyle::GetBrush("DetailsView.CategoryTop"))
+				.BorderBackgroundColor(FLinearColor::White)
+				.VAlign(EVerticalAlignment::VAlign_Center)
+				[
+					SNew(SOverlay)
+					+ SOverlay::Slot()
+					.Padding(5, 0, 0, 0)
+					[
+						SNew(STextBlock) // Title text
+						.Text(LOCTEXT("SolverSettings", "Solver Settings"))
+						.TransformPolicy(ETextTransformPolicy::ToUpper)
+						.Font(FAppStyle::Get().GetFontStyle(TEXT("PropertyWindow.BoldFont")))
+						.TextStyle(FAppStyle::Get(), "DetailsView.CategoryTextStyle")
+					]
+				]
+			]
+		]
+
+	+ SVerticalBox::Slot() // Focal Length Estimate
+		.VAlign(EVerticalAlignment::VAlign_Top)
+		.AutoHeight()
+		.MaxHeight(FCameraCalibrationWidgetHelpers::DefaultRowHeight)
+		[FCameraCalibrationWidgetHelpers::BuildLabelWidgetPair(LOCTEXT("FocalLengthEstimate", "Focal Length Estimate"), BuildFocalLengthEstimateWidget())]
+
+	+ SVerticalBox::Slot() // Fix Image Center
+		.VAlign(EVerticalAlignment::VAlign_Top)
+		.AutoHeight()
+		.MaxHeight(FCameraCalibrationWidgetHelpers::DefaultRowHeight)
+		[FCameraCalibrationWidgetHelpers::BuildLabelWidgetPair(LOCTEXT("FixImageCenter", "Fix Image Center"), BuildFixImageCenterWidget())]
+
+	+ SVerticalBox::Slot() // Calibration Rows
+		.AutoHeight()
+		.MaxHeight(12 * FCameraCalibrationWidgetHelpers::DefaultRowHeight)
+		[BuildCalibrationPointsTable()]
+
+	+ SVerticalBox::Slot() // Action buttons (e.g. Remove, Clear)
+		.HAlign(EHorizontalAlignment::HAlign_Center)
+		.AutoHeight()
+		.Padding(0, 20)
+		[BuildCalibrationActionButtons()];
+}
+
+bool UCameraLensDistortionAlgoCheckerboard::ValidateNewRow(TSharedPtr<FLensDistortionCheckerboardRowData>& Row, FText& OutErrorMessage) const
+{
+	if (!Row.IsValid())
+	{
+		OutErrorMessage = LOCTEXT("InvalidRowPointer", "Invalid row pointer");
+		return false;
+	}
+
+	if (!Tool.IsValid())
+	{
+		return false;
+	}
+
+	// Camera data is valid
+	if (!Row->CameraData.bIsValid)
+	{
+		OutErrorMessage = LOCTEXT("InvalidCameraData", "Invalid CameraData");
+		return false;
+	}
+
+	// Valid image dimensions
+	if ((Row->ImageHeight < 1) || (Row->ImageWidth < 1))
+	{
+		OutErrorMessage = LOCTEXT("InvalidImageDimensions", "Image dimensions were less than 1 pixel");
+		return false;
+	}
+
+	// Valid image pattern size
+	if (Row->SquareSideInCm < 1)
+	{
+		OutErrorMessage = LOCTEXT("InvalidPatternSize", "Pattern size cannot be smaller than 1 cm");
+		return false;
+	}
+
+	// valid number of rows and columns
+	if ((Row->NumCornerCols < 3) || (Row->NumCornerRows < 3))
+	{
+		OutErrorMessage = LOCTEXT("InvalidPatternRowsandCols", "Number of corner rows or columns cannot be less than 3");
+		return false;
+	}
+
+	// If we have no existing rows to compare this one with, we're good to go
+	if (!CalibrationRows.Num())
+	{
+		return true;
+	}
+
+	FCameraCalibrationStepsController* StepsController = Tool->GetCameraCalibrationStepsController();
+
+	if (!StepsController)
+	{
+		return false;
+	}
+
+	const ULensFile* LensFile = StepsController->GetLensFile();
+
+	if (!LensFile)
+	{
+		OutErrorMessage = LOCTEXT("LensFileNotFound", "LensFile not found");
+		return false;
+	}
+
+	const TSharedPtr<FLensDistortionCheckerboardRowData>& FirstRow = CalibrationRows[0];
+
+	// NumRows didn't change
+	if (Row->NumCornerRows != FirstRow->NumCornerRows)
+	{
+		OutErrorMessage = LOCTEXT("NumRowsChanged", "Number of corner rows changed in calibrator, which is an invalid condition.");
+		return false;
+	}
+
+	// NumCols didn't change
+	if (Row->NumCornerCols != FirstRow->NumCornerCols)
+	{
+		OutErrorMessage = LOCTEXT("NumColsChanged", "Number of corner columns changed in calibrator, which is an invalid condition.");
+		return false;
+	}
+
+	// Square side length didn't change
+	if (!FMath::IsNearlyEqual(Row->SquareSideInCm, FirstRow->SquareSideInCm))
+	{
+		OutErrorMessage = LOCTEXT("PatternSizeChanged", "Physical size of the pattern changed, which is an invalid condition");
+		return false;
+	}
+
+	// Image dimensions did not change
+	if ((Row->ImageWidth != FirstRow->ImageWidth) || (Row->ImageHeight != FirstRow->ImageHeight))
+	{
+		OutErrorMessage = LOCTEXT("ImageDimensionsChanged", "The dimensions of the media plate changed, which is an invalid condition");
+		return false;
+	}
+
+	//@todo Focus and zoom did not change much (i.e. inputs to distortion and nodal offset). 
+	//      Threshold for physical units should differ from normalized encoders.
+
+	return true;
+}
+
+FDistortionCalibrationTask UCameraLensDistortionAlgoCheckerboard::BeginCalibration(FText& OutErrorMessage)
+{
+	FDistortionCalibrationTask CalibrationTask = {};
+
+	// Validate that enough points were gathered to attempt a calibration
+	if (CalibrationRows.Num() < 1)
+	{
+		OutErrorMessage = LOCTEXT("NotEnoughCalibrationRowsError", "Could not initiate distortion calibration. At least 1 calibration row is required.");
+		return CalibrationTask;
+	}
+
+	// All points are valid
+	for (const TSharedPtr<FLensDistortionCheckerboardRowData>& Row : CalibrationRows)
+	{
+		if (!ensure(Row.IsValid()))
+		{
+			return CalibrationTask;
+		}
+	}
+
+	ULensDistortionTool* LensDistortionTool = Tool.Get();
+	if (!ensureMsgf((LensDistortionTool != nullptr), TEXT("The Lens Distortion Tool was invalid.")))
+	{
+		return CalibrationTask;
+	}
+
+	FCameraCalibrationStepsController* StepsController = LensDistortionTool->GetCameraCalibrationStepsController();
+	if (!ensureMsgf((StepsController != nullptr), TEXT("The Calibration Steps Controller was invalid.")))
+	{
+		return CalibrationTask;
+	}
+
+	ULensFile* LensFile = StepsController->GetLensFile();
+	if (!ensureMsgf((LensFile != nullptr), TEXT("The Lens File was invalid.")))
+	{
+		return CalibrationTask;
+	}
+
+	// The evaluation focus and zoom values should not change during data collection, so all rows should have the same values.
+	const float Focus = CalibrationRows[0]->CameraData.InputFocus;
+	const float Zoom = CalibrationRows[0]->CameraData.InputZoom;
+
+	const FIntPoint ImageSize = LensFile->CameraFeedInfo.GetDimensions();
+
+	const float PixelAspect = LensFile->LensInfo.SqueezeFactor;
+
+	if (FMath::IsNearlyZero(PixelAspect))
+	{
+		OutErrorMessage = LOCTEXT("PixelAspectZeroError", "The pixel aspect ratio of the CineCamera is zero, which is invalid.");
+		return CalibrationTask;
+	}
+
+	const float PhysicalSensorWidth = StepsController->GetLensFileEvaluationInputs().Filmback.SensorWidth;
+	const float DesqueezeSensorWidth = PhysicalSensorWidth * PixelAspect;
+
+	if (FMath::IsNearlyZero(DesqueezeSensorWidth))
+	{
+		OutErrorMessage = LOCTEXT("SensorWidthZeroError", "One of the filmback dimensions of the CineCamera is zero, which is invalid.");
+		return CalibrationTask;
+	}
+
+	if (!FocalLengthEstimate.IsSet() || FMath::IsNearlyZero(FocalLengthEstimate.GetValue()))
+	{
+		OutErrorMessage = LOCTEXT("FocalLengthEstimateError", "Enter a non-zero value (in mm) for the focal length estimate.");
+		return CalibrationTask;
+	}
+
+	const float FocalLengthEstimateValue = FocalLengthEstimate.GetValue();
+	const double Fx = ((double)FocalLengthEstimateValue / DesqueezeSensorWidth) * ImageSize.X;
+
+	// When operating on a desqueezed image, we expect our pixel aspect to be square, so horizontal and vertical field of view are assumed to be equal (i.e. Fx == Fy)
+	FVector2D FocalLength = FVector2D(Fx, Fx);
+	FVector2D ImageCenter = FVector2D((ImageSize.X - 1) * 0.5, (ImageSize.Y - 1) * 0.5);
+
+	TArray<FObjectPoints> Samples3d;
+	TArray<FImagePoints> Samples2d;
+	Samples3d.Reserve(CalibrationRows.Num());
+	Samples2d.Reserve(CalibrationRows.Num());
+
+	TArray<FTransform> CameraPoses;
+	CameraPoses.Reserve(CalibrationRows.Num());
+
+	TArray<FTransform> TargetPoses;
+	CameraPoses.Reserve(CalibrationRows.Num());
+
+	for (const TSharedPtr<FLensDistortionCheckerboardRowData>& Row : CalibrationRows)
+	{
+		FObjectPoints Points3d;
+		Points3d.Points = Row->Points3d;
+
+		FImagePoints Points2d;
+		Points2d.Points.Reserve(Row->Points2d.Num());
+		for (const FVector2f& Point2d : Row->Points2d)
+		{
+			Points2d.Points.Add(FVector2D(Point2d.X, Point2d.Y));
+		}
+
+		Samples3d.Add(Points3d);
+		Samples2d.Add(Points2d);
+		CameraPoses.Add(Row->CameraData.Pose);
+
+		FTransform BoardPose;
+
+		FVector TopLeft = Row->Points3d[0];
+		FVector TopRight = Row->Points3d[Row->NumCornerCols - 1];
+		FVector BottomLeft = Row->Points3d[Row->NumCornerCols * (Row->NumCornerRows - 1)];
+
+		BoardPose.SetLocation(TopLeft);
+
+		FRotator BoardRotation = FRotationMatrix::MakeFromYZ(TopRight - TopLeft, TopLeft - BottomLeft).Rotator();
+		BoardPose.SetRotation(BoardRotation.Quaternion());
+
+		TargetPoses.Add(BoardPose);
+	}
+
+	if (!CVarUseTrackedCalibrator.GetValueOnGameThread())
+	{
+		// Use a set of dummy points for the 3D checkerboard corners. The board is assumed to lie in the YZ plane with the TopLeft corner at (0, 0, 0) in world space.
+		// The side length of the checkerboard squares are used to position the remaining corners.
+		Samples3d.Empty();
+		for (int32 ImageIndex = 0; ImageIndex < CalibrationRows.Num(); ++ImageIndex)
+		{
+			FObjectPoints Points3d;
+			for (int32 RowIdx = 0; RowIdx < CalibrationRows[0]->NumCornerRows; ++RowIdx)
+			{
+				for (int32 ColIdx = 0; ColIdx < CalibrationRows[0]->NumCornerCols; ++ColIdx)
+				{
+					Points3d.Points.Add(CalibrationRows[0]->SquareSideInCm * FVector(0, ColIdx, -RowIdx));
+				}
+			}
+
+			Samples3d.Add(Points3d);
+		}
+	}
+
+	ECalibrationFlags SolverFlags = ECalibrationFlags::None;
+
+	if (CVarUseExtrinsicsGuess.GetValueOnGameThread())
+	{
+		EnumAddFlags(SolverFlags, ECalibrationFlags::UseExtrinsicGuess);
+	}
+
+	if (CVarUseIntrinsicsGuess.GetValueOnGameThread())
+	{
+		EnumAddFlags(SolverFlags, ECalibrationFlags::UseIntrinsicGuess);
+	}
+
+	if (CVarFixExtrinsics.GetValueOnGameThread())
+	{
+		EnumAddFlags(SolverFlags, ECalibrationFlags::FixExtrinsics);
+	}
+
+	if (CVarFixDistortion.GetValueOnGameThread())
+	{
+		EnumAddFlags(SolverFlags, ECalibrationFlags::FixDistortion);
+	}
+
+	if (bFixFocalLength)
+	{
+		EnumAddFlags(SolverFlags, ECalibrationFlags::FixFocalLength);
+	}
+
+	if (bFixImageCenter)
+	{
+		EnumAddFlags(SolverFlags, ECalibrationFlags::FixPrincipalPoint);
+	}
+
+	if (CVarSolveTargetOffset.GetValueOnGameThread())
+	{
+		EnumAddFlags(SolverFlags, ECalibrationFlags::SolveTargetOffset);
+	}
+
+	FDistortionInfo DistortionGuess;
+	LensFile->EvaluateDistortionParameters(Focus, Zoom, DistortionGuess);
+
+	const TSubclassOf<ULensModel> Model = LensFile->LensInfo.LensModel;
+	
+	Solver = NewObject<ULensDistortionSolver>(this, Tool->GetSolverClass());
+
+	CalibrationTask = UE::Tasks::Launch(UE_SOURCE_LOCATION, [Solver = Solver, Model, Samples3d, Samples2d, ImageSize, FocalLength, ImageCenter, DistortionGuess, CameraPoses, TargetPoses, PixelAspect, SolverFlags, Focus, Zoom]() mutable
+		{
+			FDistortionCalibrationResult Result = Solver->Solve(
+				Samples3d,
+				Samples2d,
+				ImageSize,
+				FocalLength,
+				ImageCenter,
+				DistortionGuess.Parameters,
+				CameraPoses,
+				TargetPoses,
+				Model,
+				PixelAspect,
+				SolverFlags
+			);
+
+			// CalibrateCamera() returns focal length and image center in pixels, but the result is expected to be normalized by the image size
+			Result.FocalLength.FxFy = Result.FocalLength.FxFy / ImageSize;
+			Result.ImageCenter.PrincipalPoint = Result.ImageCenter.PrincipalPoint / ImageSize;
+
+			// FZ inputs to LUT
+			Result.EvaluatedFocus = Focus;
+			Result.EvaluatedZoom = Zoom;
+
+			return Result;
+		});
+
+	return CalibrationTask;
+}
+
+void UCameraLensDistortionAlgoCheckerboard::CancelCalibration()
+{
+	if (Solver)
+	{
+		Solver->Cancel();
+	}
+}
+
+bool UCameraLensDistortionAlgoCheckerboard::GetCalibrationStatus(FText& StatusText) const
+{
+	if (Solver)
+	{
+		return Solver->GetStatusText(StatusText);
+	}
+	return false;
+}
+
+TSharedRef<SWidget> UCameraLensDistortionAlgoCheckerboard::BuildCalibrationDevicePickerWidget()
+{
+	return SNew(SHorizontalBox)
+
+	+ SHorizontalBox::Slot() // Picker
+	[
+		SNew(SObjectPropertyEntryBox)
+		.AllowedClass(ACameraCalibrationCheckerboard::StaticClass())
+		.OnObjectChanged_Lambda([&](const FAssetData& AssetData) -> void
+		{
+			if (AssetData.IsValid())
+			{
+				SetCalibrator(Cast<ACameraCalibrationCheckerboard>(AssetData.GetAsset()));
+			}
+		})
+		.ObjectPath_Lambda([&]() -> FString
+		{
+			if (AActor* TheCalibrator = GetCalibrator())
+			{
+				FAssetData AssetData(TheCalibrator, true);
+				return AssetData.GetObjectPathString();
+			}
+
+			return TEXT("");
+		})
+	]
+
+	+ SHorizontalBox::Slot() // Spawner
+	.AutoWidth()
+	[
+		SNew(SButton)
+		.Text(LOCTEXT("Spawn", "Spawn"))
+		.HAlign(HAlign_Center)
+		.VAlign(VAlign_Center)
+		.ButtonStyle(FAppStyle::Get(), "HoverHintOnly")
+		.OnClicked_Lambda([&]() -> FReply
+		{
+			const FCameraCalibrationStepsController* StepsController = GetStepsController();
+
+			if (!ensure(StepsController))
+			{
+				return FReply::Handled();
+			}
+
+			if (UWorld* const World = StepsController->GetWorld())
+			{
+				SetCalibrator(World->SpawnActor<ACameraCalibrationCheckerboard>());
+			}
+
+			return FReply::Handled();
+		})
+		[
+			SNew(STextBlock)
+			.Font(FAppStyle::Get().GetFontStyle("FontAwesome.12"))
+			.Text(FEditorFontGlyphs::Plus)
+			.ColorAndOpacity(FLinearColor::White)
+		]
+	];
+}
+
+TSharedRef<SWidget> UCameraLensDistortionAlgoCheckerboard::BuildShowOverlayWidget()
+{
+	return SNew(SCheckBox)
+	.IsChecked_Lambda([&]() -> ECheckBoxState
+	{
+		return bShouldShowOverlay ? ECheckBoxState::Checked : ECheckBoxState::Unchecked;
+	})
+	.OnCheckStateChanged_Lambda([&](ECheckBoxState NewState) -> void
+	{
+		bShouldShowOverlay = (NewState == ECheckBoxState::Checked);
+
+		FCameraCalibrationStepsController* StepsController = Tool->GetCameraCalibrationStepsController();
+		
+		if (!ensure(StepsController))
+		{
+			return;
+		}
+
+		StepsController->SetOverlayEnabled(bShouldShowOverlay);
+	});
+}
+
+TSharedRef<SWidget> UCameraLensDistortionAlgoCheckerboard::BuildShowDetectionWidget()
+{
+	return SNew(SCheckBox)
+	.IsChecked_Lambda([&]() -> ECheckBoxState
+	{
+		return bShouldShowDetectionWindow ? ECheckBoxState::Checked : ECheckBoxState::Unchecked;
+	})
+	.OnCheckStateChanged_Lambda([&](ECheckBoxState NewState) -> void
+	{
+		bShouldShowDetectionWindow = (NewState == ECheckBoxState::Checked);
+	});
+}
+
+TSharedRef<SWidget> UCameraLensDistortionAlgoCheckerboard::BuildFocalLengthEstimateWidget()
+{
+	return SNew(SHorizontalBox)
+	
+	+ SHorizontalBox::Slot()
+	.FillWidth(0.8)
+	.Padding(0, 0, 10, 0)
+	[
+		SNew(SNumericEntryBox<float>)
+		.Value_UObject(this, &UCameraLensDistortionAlgoCheckerboard::GetFocalLengthEstimate)
+		.OnValueChanged_UObject(this, &UCameraLensDistortionAlgoCheckerboard::SetFocalLengthEstimate)
+	]
+
+	+ SHorizontalBox::Slot()
+	.FillWidth(0.2)
+	[
+		SNew(SCheckBox)
+		.Padding(FMargin(5, 0, 15, 0))
+		.IsChecked_UObject(this, &UCameraLensDistortionAlgoCheckerboard::IsFixFocalLengthChecked)
+		.OnCheckStateChanged_UObject(this, &UCameraLensDistortionAlgoCheckerboard::OnFixFocalLengthCheckStateChanged)
+		[
+			SNew(STextBlock).Text(LOCTEXT("FixText", "Fix?"))
+		]
+	];
+}
+
+TOptional<float> UCameraLensDistortionAlgoCheckerboard::GetFocalLengthEstimate() const
+{
+	return FocalLengthEstimate;
+}
+
+void UCameraLensDistortionAlgoCheckerboard::SetFocalLengthEstimate(float NewValue)
+{
+	FocalLengthEstimate = NewValue;
+}
+
+ECheckBoxState UCameraLensDistortionAlgoCheckerboard::IsFixFocalLengthChecked() const
+{
+	return bFixFocalLength ? ECheckBoxState::Checked : ECheckBoxState::Unchecked;
+}
+
+void UCameraLensDistortionAlgoCheckerboard::OnFixFocalLengthCheckStateChanged(ECheckBoxState NewState)
+{
+	bFixFocalLength = (NewState == ECheckBoxState::Checked);
+}
+
+TSharedRef<SWidget> UCameraLensDistortionAlgoCheckerboard::BuildFixImageCenterWidget()
+{
+	return SNew(SCheckBox)
+	.IsChecked_Lambda([&]() -> ECheckBoxState
+	{
+		return bFixImageCenter ? ECheckBoxState::Checked : ECheckBoxState::Unchecked;
+	})
+	.OnCheckStateChanged_Lambda([&](ECheckBoxState NewState) -> void
+	{
+		bFixImageCenter = (NewState == ECheckBoxState::Checked);
+	});
+}
+
+TSharedRef<SWidget> UCameraLensDistortionAlgoCheckerboard::BuildExtrinsicsGuessWidget()
+{
+	return SNew(SCheckBox)
+	.IsChecked_Lambda([&]() -> ECheckBoxState
+	{
+		return bUseExtrinsicsGuess ? ECheckBoxState::Checked : ECheckBoxState::Unchecked;
+	})
+	.OnCheckStateChanged_Lambda([&](ECheckBoxState NewState) -> void
+	{
+		bUseExtrinsicsGuess = (NewState == ECheckBoxState::Checked);
+	});
+}
+
+TSharedRef<SWidget> UCameraLensDistortionAlgoCheckerboard::BuildCalibrateNodalOffsetWidget()
+{
+	return SNew(SCheckBox)
+	.IsChecked_Lambda([&]() -> ECheckBoxState
+	{
+		return bCalibrateNodalOffset ? ECheckBoxState::Checked : ECheckBoxState::Unchecked;
+	})
+	.OnCheckStateChanged_Lambda([&](ECheckBoxState NewState) -> void
+	{
+		bCalibrateNodalOffset = (NewState == ECheckBoxState::Checked);
+	});
+}
+
+TSharedRef<SWidget> UCameraLensDistortionAlgoCheckerboard::BuildCalibrationActionButtons()
+{
+	return SNew(SHorizontalBox)
+
+	+ SHorizontalBox::Slot() // Button to clear all rows
+	.AutoWidth()
+	[
+		SNew(SButton)
+		.Text(LOCTEXT("ClearAll", "Clear All"))
+		.HAlign(HAlign_Center)
+		.VAlign(VAlign_Center)
+		.OnClicked_Lambda([&]() -> FReply
+		{
+			ClearCalibrationRows();
+			return FReply::Handled();
+		})
+	];
+}
+
+TSharedRef<SWidget> UCameraLensDistortionAlgoCheckerboard::BuildCalibrationPointsTable()
+{
+	CalibrationListView = SNew(SListView<TSharedPtr<FLensDistortionCheckerboardRowData>>)
+		.ListItemsSource(&CalibrationRows)
+		.OnGenerateRow_Lambda([&](TSharedPtr<FLensDistortionCheckerboardRowData> InItem, const TSharedRef<STableViewBase>& OwnerTable) -> TSharedRef<ITableRow>
+		{
+			return SNew(CameraLensDistortionAlgoCheckerboard::SCalibrationRowGenerator, OwnerTable)
+			.CalibrationRowData(InItem);
+		})
+		.SelectionMode(ESelectionMode::Multi)
+		.OnKeyDownHandler_Lambda([&](const FGeometry& Geometry, const FKeyEvent& KeyEvent) -> FReply
+		{
+			if (!CalibrationListView.IsValid())
+			{
+				return FReply::Unhandled();
+			}
+
+			if (KeyEvent.GetKey() == EKeys::Delete)
+			{
+				// Delete selected items
+
+				const TArray<TSharedPtr<FLensDistortionCheckerboardRowData>> SelectedItems = CalibrationListView->GetSelectedItems();
+
+				for (const TSharedPtr<FLensDistortionCheckerboardRowData>& SelectedItem : SelectedItems)
+				{
+					CalibrationRows.Remove(SelectedItem);
+
+					// Delete the previously exported .json file for the row that is being deleted
+					if (ULensDistortionTool* LensDistortionTool = Tool.Get())
+					{
+						LensDistortionTool->DeleteExportedRow(SelectedItem->Index);
+					}
+				}
+
+				CalibrationListView->RequestListRefresh();
+				RefreshCoverage();
+				return FReply::Handled();
+			}
+			else if (KeyEvent.GetModifierKeys().IsControlDown() && (KeyEvent.GetKey() == EKeys::A))
+			{
+				// Select all items
+
+				CalibrationListView->SetItemSelection(CalibrationRows, true);
+				return FReply::Handled();
+			}
+			else if (KeyEvent.GetKey() == EKeys::Escape)
+			{
+				// Deselect all items
+
+				CalibrationListView->ClearSelection();
+				return FReply::Handled();
+			}
+
+			return FReply::Unhandled();
+		})
+		.HeaderRow
+		(
+			SNew(SHeaderRow)
+
+			+ SHeaderRow::Column("Index")
+			.DefaultLabel(LOCTEXT("Index", "Index"))
+
+			+ SHeaderRow::Column("Image")
+			.DefaultLabel(LOCTEXT("Image", "Image"))
+		);
+
+	return CalibrationListView.ToSharedRef();
+}
+
+ACameraCalibrationCheckerboard* UCameraLensDistortionAlgoCheckerboard::FindFirstCalibrator() const
+{
+	if (!ensure(Tool.IsValid()))
+	{
+		return nullptr;
+	}
+
+	const FCameraCalibrationStepsController* StepsController = Tool->GetCameraCalibrationStepsController();
+
+	if (!ensure(StepsController))
+	{
+		return nullptr;
+	}
+
+	UWorld* World = StepsController->GetWorld();
+
+	if (!ensure(World))
+	{
+		return nullptr;
+	}
+
+	const EActorIteratorFlags Flags = EActorIteratorFlags::SkipPendingKill;
+
+	for (TActorIterator<AActor> It(World, ACameraCalibrationCheckerboard::StaticClass(), Flags); It; ++It)
+	{
+		return CastChecked<ACameraCalibrationCheckerboard>(*It);
+	}
+
+	return nullptr;
+}
+
+void UCameraLensDistortionAlgoCheckerboard::SetCalibrator(ACameraCalibrationCheckerboard* InCalibrator)
+{
+	Calibrator = InCalibrator;
+}
+
+ACameraCalibrationCheckerboard* UCameraLensDistortionAlgoCheckerboard::GetCalibrator() const
+{
+	return Calibrator.Get();
+}
+
+UMaterialInterface* UCameraLensDistortionAlgoCheckerboard::GetOverlayMaterial() const
+{
+	return TSoftObjectPtr<UMaterialInterface>(FSoftObjectPath(TEXT("/CameraCalibration/Materials/M_Coverage.M_Coverage"))).LoadSynchronous();
+}
+
+void UCameraLensDistortionAlgoCheckerboard::OnDistortionSavedToLens()
+{
+	// Since the calibration result was saved, there is no further use for the current samples.
+	ClearCalibrationRows();
+}
+
+void UCameraLensDistortionAlgoCheckerboard::ClearCalibrationRows()
+{
+	CalibrationRows.Empty();
+
+	if (CalibrationListView.IsValid())
+	{
+		CalibrationListView->RequestListRefresh();
+	}
+
+	RefreshCoverage();
+
+	// End the current calibration session (a new one will begin the next time a new row is added)
+	if (ULensDistortionTool* LenDistortionTool = Tool.Get())
+	{
+		LenDistortionTool->EndCalibrationSession();
+	}
+}
+
+void UCameraLensDistortionAlgoCheckerboard::RefreshCoverage()
+{
+	FCameraCalibrationStepsController* StepsController = Tool->GetCameraCalibrationStepsController();
+
+	if (!ensure(StepsController))
+	{
+		return;
+	}
+
+	UE::CameraCalibration::Private::ClearTexture(CoverageTexture);
+
+	for (const TSharedPtr<FLensDistortionCheckerboardRowData>& Row : CalibrationRows)
+	{
+		TArray<FVector2f> CameraFeedAdjustedCorners = Row->Points2d;
+		const FIntPoint CoverageTextureSize = FIntPoint(CoverageTexture->GetSizeX(), CoverageTexture->GetSizeY());
+		const FIntPoint ImageSize = FIntPoint(Row->ImageWidth, Row->ImageHeight);
+		UE::CameraCalibration::Private::AdjustCorners(CameraFeedAdjustedCorners, CoverageTextureSize, ImageSize);
+
+		const FIntPoint CheckerboardDimensions = FIntPoint(Row->NumCornerCols, Row->NumCornerRows);
+		FOpenCVHelper::DrawCheckerboardCorners(Row->Points2d, CheckerboardDimensions, CoverageTexture);
+	}
+
+	// The coverage texture may have changed as a result of a change in size or pixel format.
+	// Therefore, the material parameter should be updated to ensure it is up to date.
+	if (CoverageTexture)
+	{
+		if (UMaterialInstanceDynamic* OverlayMID = Tool->GetOverlayMID())
+		{
+			OverlayMID->SetTextureParameterValue(FName(TEXT("CoverageTexture")), CoverageTexture);
+		}
+	}
+
+	StepsController->RefreshOverlay();
+}
+
+
+void UCameraLensDistortionAlgoCheckerboard::ExportSessionData()
+{
+	using namespace UE::CameraCalibration::Private;
+
+	if (ULensDistortionTool* LensDistortionTool = Tool.Get())
+	{
+		// Add all data to a json object that is needed to run this algorithm that is NOT part of a specific row
+		TSharedPtr<FJsonObject> JsonSessionData = MakeShared<FJsonObject>();
+
+		JsonSessionData->SetNumberField(LensDistortionCheckerboardExportFields::Version, DATASET_VERSION);
+
+		// Export the session data to a .json file
+		Tool->ExportSessionData(JsonSessionData.ToSharedRef());
+	}
+}
+
+void UCameraLensDistortionAlgoCheckerboard::ExportRow(TSharedPtr<FLensDistortionCheckerboardRowData> Row)
+{
+	if (ULensDistortionTool* LensDistortionTool = Tool.Get())
+	{
+		if (const TSharedPtr<FJsonObject>& RowObject = FJsonObjectConverter::UStructToJsonObject<FLensDistortionCheckerboardRowData>(Row.ToSharedRef().Get()))
+		{
+			// Export the row data to a .json file
+			LensDistortionTool->ExportCalibrationRow(Row->Index, RowObject.ToSharedRef(), Row->ImageView);
+		}
+	}
+}
+
+bool UCameraLensDistortionAlgoCheckerboard::HasCalibrationData() const
+{
+	return (CalibrationRows.Num() > 0);
+}
+
+void UCameraLensDistortionAlgoCheckerboard::PreImportCalibrationData()
+{
+	// Clear the current set of rows before importing new ones
+	CalibrationRows.Empty();
+}
+
+int32 UCameraLensDistortionAlgoCheckerboard::ImportCalibrationRow(const TSharedRef<FJsonObject>& CalibrationRowObject, const FImage& RowImage)
+{
+	// Create a new row to populate with data from the Json object
+	TSharedPtr<FLensDistortionCheckerboardRowData> NewRow = MakeShared<FLensDistortionCheckerboardRowData>();
+
+	// We enforce strict mode to ensure that every field in the UStruct of row data is present in the imported json.
+	// If any fields are missing, it is likely the row will be invalid, which will lead to errors in the calibration.
+	constexpr int64 CheckFlags = 0;
+	constexpr int64 SkipFlags = 0;
+	constexpr bool bStrictMode = true;
+	if (FJsonObjectConverter::JsonObjectToUStruct<FLensDistortionCheckerboardRowData>(CalibrationRowObject, NewRow.Get(), CheckFlags, SkipFlags, bStrictMode))
+	{
+		if (!RowImage.RawData.IsEmpty())
+		{
+			NewRow->ImageView = RowImage;
+
+			// Resize the image into a thumbnail, create a UTexture from it, and initialize the thumbnail viewport for this row
+			FImage RowThumbnail;
+			RowImage.ResizeTo(RowThumbnail, RowImage.GetWidth() / 4, RowImage.GetHeight() / 4, ERawImageFormat::BGRA8, RowImage.GetGammaSpace());
+
+			UTexture2D* ThumbnailTexture = FImageUtils::CreateTexture2DFromImage(RowThumbnail);
+			NewRow->Thumbnail = SNew(SSimulcamViewport, ThumbnailTexture).WithZoom(false).WithPan(false);
+		}
+
+		CalibrationRows.Add(NewRow);
+	}
+	else
+	{
+		UE_LOG(LogCameraCalibrationEditor, Warning, TEXT("LensDistortionCheckerboard algo failed to import calibration row because at least one field could not be deserialized from the json file."));
+	}
+
+	return NewRow->Index;
+}
+
+void UCameraLensDistortionAlgoCheckerboard::PostImportCalibrationData()
+{
+	// Sort imported calibration rows by row index
+	CalibrationRows.Sort([](const TSharedPtr<FLensDistortionCheckerboardRowData>& LHS, const TSharedPtr<FLensDistortionCheckerboardRowData>& RHS) { return LHS->Index < RHS->Index; });
+
+	// Notify the ListView of the new data
+	if (CalibrationListView.IsValid())
+	{
+		CalibrationListView->RequestListRefresh();
+	}
+}
+
+TSharedRef<SWidget> UCameraLensDistortionAlgoCheckerboard::BuildHelpWidget()
+{
+	return SNew(STextBlock)
+		.Text(LOCTEXT("LensDistortionCheckerboardHelp",
+			"This Lens Distortion algorithm is based on capturing at least 4 different views of a checkerboard.\n\n"
+			"The camera and/or the checkerboard may be moved for each. To capture, simply click the simulcam\n"
+			"viewport. You can optionally right-click the simulcam viewport to pause it and ensure it will be\n"
+			"a sharp capture (if the media source supports pause).\n\n"
+
+			"This will require the selection of a Checkerboard actor that exists in the scene, and is configured\n"
+			"to match the dimensions and number of inner corner rows and columns as the physical checkerboard.\n"
+			"If there isn't a checkerboard in the scene, you can spawn one from either the Place Actors panel\n"
+			"or the '+' button next to the checkerboard picker. You can then edit its details panel properties.\n\n"
+
+			"Each capture will be added to the table with the thumbnails. It is important that the collection of\n"
+			"captures sweep the viewport area as much as possible, so that the calibration includes samples from\n"
+			"all of its regions (otherwise the calibration may be skewed towards a particular region and not be\n"
+			"very accurate in others).\n\n"
+
+			"When you are done capturing, click the 'Add To Lens Distortion Calibration' button to run the calibration\n"
+			"algorithm, which calculates the distortion parameters k1,k2,p1,p2,k3 and adds them to the lens file.\n\n"
+
+			"A dialog window will show you the reprojection error of the calibration, and you can decide to add the\n"
+			"data to the lens file or not.\n"
+		));
+}
+
+#undef LOCTEXT_NAMESPACE
diff --git a/Engine/Plugins/VirtualProduction/CameraCalibration/Source/CameraCalibrationEditor/Private/Calibrators/CameraLensDistortionAlgoCheckerboard.h b/Engine/Plugins/VirtualProduction/CameraCalibration/Source/CameraCalibrationEditor/Private/Calibrators/CameraLensDistortionAlgoCheckerboard.h
new file mode 100644
--- /dev/null
+++ b/Engine/Plugins/VirtualProduction/CameraCalibration/Source/CameraCalibrationEditor/Private/Calibrators/CameraLensDistortionAlgoCheckerboard.h
@@ -0,0 +1,272 @@
+// Copyright Epic Games, Inc. All Rights Reserved.
+
+#pragma once
+
+#include "CameraLensDistortionAlgo.h"
+
+#include "CameraCalibrationSolver.h"
+#include "ImageCore.h"
+
+#include "CameraLensDistortionAlgoCheckerboard.generated.h"
+
+struct FGeometry;
+struct FPointerEvent;
+
+class ACameraCalibrationCheckerboard;
+class FCameraCalibrationStepsController;
+class ULensDistortionTool;
+class SSimulcamViewport;
+
+template <typename ItemType>
+class SListView;
+
+class SWidget;
+class UMediaTexture;
+class UWorld;
+class FJsonObject;
+
+template<typename OptionType>
+class SComboBox;
+
+namespace CameraLensDistortionAlgoCheckerboard
+{
+	class SCalibrationRowGenerator;
+}
+
+/** Holds camera information that can be used to add the samples */
+USTRUCT()
+struct FLensDistortionCheckerboardCameraData
+{
+	GENERATED_BODY()
+
+	// True if the rest of the contents are valid.
+	UPROPERTY()
+	bool bIsValid = false;
+
+	// Input focus from lens file evaluation data
+	UPROPERTY()
+	float InputFocus = 0.0f;
+
+	// Input zoom from lens file evaluation data
+	UPROPERTY()
+	float InputZoom = 0.0f;
+
+	// The camera pose
+	UPROPERTY()
+	FTransform Pose = FTransform::Identity;
+
+	// True if the cached camera pose included a nodal offset
+	UPROPERTY()
+	bool bWasNodalOffsetApplied = false;
+};
+
+/** Holds information of the calibration row */
+USTRUCT()
+struct FLensDistortionCheckerboardRowData
+{
+	GENERATED_BODY()
+
+	// Thumbnail to display in list
+	TSharedPtr<SSimulcamViewport> Thumbnail;
+
+	// Captured image containing checkerboard
+	FImageView ImageView;
+
+	// Index to display in list
+	UPROPERTY()
+	int32 Index = -1;
+
+	// Checkerboard corners in 2d image pixel coordinates.
+	UPROPERTY()
+	TArray<FVector2f> Points2d;
+
+	// Checkerboard corners in 3d local space.
+	UPROPERTY()
+	TArray<FVector> Points3d;
+
+	// Which calibrator was used
+	UPROPERTY()
+	FString CalibratorName;
+
+	// Number of corner rows in pattern
+	UPROPERTY()
+	int32 NumCornerRows = 0;
+
+	// Number of corner columns in pattern
+	UPROPERTY()
+	int32 NumCornerCols = 0;
+
+	// Side dimension in cm
+	UPROPERTY()
+	float SquareSideInCm = 0;
+
+	// Width of image
+	UPROPERTY()
+	int32 ImageWidth = 0;
+
+	// Height of image
+	UPROPERTY()
+	int32 ImageHeight = 0;
+
+	// Holds information of the camera data for this sample
+	UPROPERTY()
+	FLensDistortionCheckerboardCameraData CameraData;
+};
+
+/** 
+ * Implements a lens distortion calibration algorithm. It requires a checkerboard pattern
+ */
+UCLASS()
+class UCameraLensDistortionAlgoCheckerboard : public UCameraLensDistortionAlgo
+{
+	GENERATED_BODY()
+	
+public:
+
+	//~ Begin CameraLensDistortionAlgo
+	virtual void Initialize(ULensDistortionTool* InTool) override;
+	virtual void Shutdown() override;
+	virtual bool SupportsModel(const TSubclassOf<ULensModel>& LensModel) const override;
+	virtual void Tick(float DeltaTime) override;
+	virtual bool OnViewportClicked(const FGeometry& MyGeometry, const FPointerEvent& MouseEvent) override;
+	virtual TSharedRef<SWidget> BuildUI() override;
+	virtual FName FriendlyName() const override { return TEXT("Lens Distortion Checkerboard"); };
+	virtual FName ShortName() const override { return TEXT("Checkerboard"); };
+	virtual UMaterialInterface* GetOverlayMaterial() const override;
+	virtual bool IsOverlayEnabled() const override { return bShouldShowOverlay; };
+	virtual void OnDistortionSavedToLens() override;
+	virtual FDistortionCalibrationTask BeginCalibration(FText& OutErrorMessage) override;
+	virtual void CancelCalibration() override;
+	virtual bool GetCalibrationStatus(FText& StatusText) const override;
+	virtual bool SupportsAsyncCalibration() override { return true; };
+	virtual bool HasCalibrationData() const override;
+	virtual void PreImportCalibrationData() override;
+	virtual int32 ImportCalibrationRow(const TSharedRef<FJsonObject>& CalibrationRowObject, const FImage& RowImage) override;
+	virtual void PostImportCalibrationData() override;
+	virtual TSharedRef<SWidget> BuildHelpWidget() override;
+	//~ End CameraLensDistortionAlgo
+
+private:
+
+	// SCalibrationRowGenerator will need access to the row structures below.
+	friend class CameraLensDistortionAlgoCheckerboard::SCalibrationRowGenerator;
+
+private:
+	/** Version of the current calibration dataset */
+	static const int DATASET_VERSION;
+
+	/** The lens distortion tool controller */
+	TWeakObjectPtr<ULensDistortionTool> Tool;
+
+	/** The currently selected checkerboard object. */
+	TWeakObjectPtr<ACameraCalibrationCheckerboard> Calibrator;
+
+	/** Rows source for the CalibrationListView */
+	TArray<TSharedPtr<FLensDistortionCheckerboardRowData>> CalibrationRows;
+
+	/** Displays the list of calibration points that will be used to calculate the lens distortion */
+	TSharedPtr<SListView<TSharedPtr<FLensDistortionCheckerboardRowData>>> CalibrationListView;
+
+	/** Caches the last camera data.  Will hold last value before the media is paused */
+	FLensDistortionCheckerboardCameraData LastCameraData;
+
+	/** True if the coverage overlay should be shown */
+	bool bShouldShowOverlay = false;
+
+	/** True if a detection window should be shown after every capture */
+	bool bShouldShowDetectionWindow = false;
+
+	/** If true, the solver will not solve for focal length, and will just use the input value */
+	bool bFixFocalLength = false;
+
+	/** If true, the solver will not solve for image center, and will just use the input value */
+	bool bFixImageCenter = false;
+
+	/** Estimate for the focal length to provide to the solver. If not set, user will be warned to set before continuing calibration */
+	TOptional<float> FocalLengthEstimate;
+
+	/** If true, the solver will use the current camera pose to initialize the camera extrinsic parameters for each image */
+	bool bUseExtrinsicsGuess = false;
+
+	/** If true, the final camera poses generated by the solver will be used to compute the nodal offset from the initial camera poses  */
+	bool bCalibrateNodalOffset = false;
+
+	/** Texture into which detected chessboard corners from each calibration row are drawn */
+	TObjectPtr<UTexture2D> CoverageTexture;
+
+	/** Solver instance that will run the distortion calibration */
+	TObjectPtr<ULensDistortionSolver> Solver;
+
+private:
+
+	/** Builds the UI of the calibration device picker */
+	TSharedRef<SWidget> BuildCalibrationDevicePickerWidget();
+
+	/** Builds the UI for the user to select if they want to display the coverage overlay */
+	TSharedRef<SWidget> BuildShowOverlayWidget();
+
+	/** Builds the UI for the user to select if they want a corner detection window to be shown after every capture */
+	TSharedRef<SWidget> BuildShowDetectionWidget();
+
+	/** Builds the UI of the focal length estimate widget */
+	TSharedRef<SWidget> BuildFocalLengthEstimateWidget();
+
+	/** Builds the UI of the fix image center widget */
+	TSharedRef<SWidget> BuildFixImageCenterWidget();
+
+	/** Builds the UI of the extrinsics guess widget */
+	TSharedRef<SWidget> BuildExtrinsicsGuessWidget();
+
+	/** Builds the UI of the calibrate nodal offset widget */
+	TSharedRef<SWidget> BuildCalibrateNodalOffsetWidget();
+
+	/** Builds the UI of the calibration points table */
+	TSharedRef<SWidget> BuildCalibrationPointsTable();
+
+	/** Builds the UI for the action buttons (RemoveLast, ClearAll) */
+	TSharedRef<SWidget> BuildCalibrationActionButtons();
+
+	/** Get the current value of the focal length estimate */
+	TOptional<float> GetFocalLengthEstimate() const;
+
+	/** Set the current value of the focal length estimate */
+	void SetFocalLengthEstimate(float NewValue);
+
+	/** Returns the state of the fix focal length checkbox */
+	ECheckBoxState IsFixFocalLengthChecked() const;
+
+	/** Changes the state of the fix focal length checkbox */
+	void OnFixFocalLengthCheckStateChanged(ECheckBoxState NewState);
+
+private:
+
+	/** Returns the first checkerboard object in the scene that it can find */
+	ACameraCalibrationCheckerboard* FindFirstCalibrator() const;
+
+	/** Sets the calibrator object to be used. Updates the selection picker. */
+	void SetCalibrator(ACameraCalibrationCheckerboard* InCalibrator);
+
+	/** Returns the currently selected calibrator object. */
+	ACameraCalibrationCheckerboard* GetCalibrator() const;
+
+	/** Clears the list of calibration sample rows */
+	void ClearCalibrationRows();
+
+	/** Validates a new calibration point to determine if it should be added as a new sample row */
+	bool ValidateNewRow(TSharedPtr<FLensDistortionCheckerboardRowData>& Row, FText& OutErrorMessage) const;
+
+	/** Add a new calibration row from media texture and camera data */
+	bool AddCalibrationRow(FText& OutErrorMessage);
+
+	/** Returns the steps controller */
+	FCameraCalibrationStepsController* GetStepsController() const;
+
+	/** Update the coverage texture to reflect the current set of calibration rows */
+	void RefreshCoverage();
+
+	/** Export global session data to a .json file */
+	void ExportSessionData();
+
+	/** Export the row data to a json file */
+	void ExportRow(TSharedPtr<FLensDistortionCheckerboardRowData> Row);
+};
diff --git a/Engine/Plugins/VirtualProduction/CameraCalibration/Source/CameraCalibrationEditor/Private/Calibrators/CameraLensDistortionAlgoPoints.cpp b/Engine/Plugins/VirtualProduction/CameraCalibration/Source/CameraCalibrationEditor/Private/Calibrators/CameraLensDistortionAlgoPoints.cpp
new file mode 100644
--- /dev/null
+++ b/Engine/Plugins/VirtualProduction/CameraCalibration/Source/CameraCalibrationEditor/Private/Calibrators/CameraLensDistortionAlgoPoints.cpp
@@ -0,0 +1,1077 @@
+// Copyright Epic Games, Inc. All Rights Reserved.
+
+#include "CameraLensDistortionAlgoPoints.h"
+
+#include "AssetEditor/CameraCalibrationStepsController.h"
+#include "AssetEditor/LensDistortionTool.h"
+#include "CalibrationPointComponent.h"
+#include "Camera/CameraActor.h"
+#include "CameraCalibrationEditorLog.h"
+#include "CameraCalibrationSolver.h"
+#include "CameraCalibrationUtilsPrivate.h"
+#include "Dom/JsonObject.h"
+#include "Editor.h"
+#include "Input/Events.h"
+#include "JsonObjectConverter.h"
+#include "Math/Vector.h"
+#include "Misc/MessageDialog.h"
+#include "Models/SphericalLensModel.h"
+#include "OpenCVHelper.h"
+#include "UI/SFilterableActorPicker.h"
+#include "UI/CameraCalibrationWidgetHelpers.h"
+#include "UObject/UObjectIterator.h"
+#include "Widgets/Input/SButton.h"
+#include "Widgets/Input/SNumericEntryBox.h"
+#include "Widgets/Views/SListView.h"
+
+
+#define LOCTEXT_NAMESPACE "CameraLensDistortionAlgoPoints"
+
+#if WITH_EDITOR
+static TAutoConsoleVariable<bool> CVarFixExtrinsicsPoints(TEXT("LensDistortionPoints.FixExtrinsics"), false, TEXT("If true, the solver will fix the camera extrinsics to the user-provided camera poses"));
+static TAutoConsoleVariable<bool> CVarFixDistortionPoints(TEXT("LensDistortionPoints.FixDistortion"), false, TEXT("If true, the solver will not optimize distortion, and will use the input distortion values if any are given, or assume zero distortion otherwise."));
+static TAutoConsoleVariable<bool> CVarUseExtrinsicsGuessPoints(TEXT("LensDistortionPoints.UseExtrinsicsGuess"), false, TEXT("If true, the actual checkerboard and camera poses will be used when running the solver"));
+#endif
+
+const int UCameraLensDistortionAlgoPoints::DATASET_VERSION = 1;
+
+// String constants for import/export
+namespace UE::CameraCalibration::Private::LensDistortionPointsExportFields
+{
+	static const FString Version(TEXT("Version"));
+	static const FString CurrentFocalLength(TEXT("CurrentFocalLength"));
+}
+
+namespace CameraLensDistortionAlgoPoints
+{
+	class SCalibrationRowGenerator : public SMultiColumnTableRow<TSharedPtr<FLensDistortionPointsRowData>>
+	{
+
+	public:
+		SLATE_BEGIN_ARGS(SCalibrationRowGenerator) {}
+
+		/** The list item for this row */
+		SLATE_ARGUMENT(TSharedPtr<FLensDistortionPointsRowData>, CalibrationRowData)
+
+		SLATE_END_ARGS()
+
+		void Construct(const FArguments& Args, const TSharedRef<STableViewBase>& OwnerTableView)
+		{
+			CalibrationRowData = Args._CalibrationRowData;
+
+			SMultiColumnTableRow<TSharedPtr<FLensDistortionPointsRowData>>::Construct(
+				FSuperRowType::FArguments()
+				.Padding(1.0f),
+				OwnerTableView
+			);
+		}
+
+		//~Begin SMultiColumnTableRow
+		virtual TSharedRef<SWidget> GenerateWidgetForColumn(const FName& ColumnName) override
+		{
+			if (ColumnName == TEXT("Pattern"))
+			{
+				const FString Text = FString::Printf(TEXT("%d"), CalibrationRowData->PatternIndex + 1);
+				return SNew(STextBlock).Text(FText::FromString(Text));
+			}
+
+			if (ColumnName == TEXT("Name"))
+			{
+				const FString Text = CalibrationRowData->CalibratorPointData.Name;
+				return SNew(STextBlock).Text(FText::FromString(Text));
+			}
+
+			if (ColumnName == TEXT("Point2D"))
+			{
+				const FString Text = FString::Printf(TEXT("(%.2f, %.2f)"),
+					CalibrationRowData->CalibratorPointData.Point2d.X,
+					CalibrationRowData->CalibratorPointData.Point2d.Y);
+
+				return SNew(STextBlock).Text(FText::FromString(Text));
+			}
+
+			if (ColumnName == TEXT("Point3D"))
+			{
+				const FString Text = FString::Printf(TEXT("(%.0f, %.0f, %.0f)"),
+					CalibrationRowData->CalibratorPointData.Point3d.X,
+					CalibrationRowData->CalibratorPointData.Point3d.Y,
+					CalibrationRowData->CalibratorPointData.Point3d.Z);
+
+				return SNew(STextBlock).Text(FText::FromString(Text));
+			}
+
+			return SNullWidget::NullWidget;
+		}
+		//~End SMultiColumnTableRow
+
+
+	private:
+		TSharedPtr<FLensDistortionPointsRowData> CalibrationRowData;
+	};
+};
+
+void UCameraLensDistortionAlgoPoints::Initialize(ULensDistortionTool* InLensDistortionTool)
+{
+	LensDistortionTool = InLensDistortionTool;
+
+	// Guess which calibrator to use by searching for actors with CalibrationPointComponents.
+	SetCalibrator(FindFirstCalibrator());
+}
+
+void UCameraLensDistortionAlgoPoints::Shutdown()
+{
+	LensDistortionTool.Reset();
+}
+
+bool UCameraLensDistortionAlgoPoints::SupportsModel(const TSubclassOf<ULensModel>& LensModel) const
+{
+	return (LensModel == USphericalLensModel::StaticClass());
+}
+
+void UCameraLensDistortionAlgoPoints::Tick(float DeltaTime)
+{
+	if (!LensDistortionTool.IsValid())
+	{
+		return;
+	}
+
+	const FCameraCalibrationStepsController* StepsController = LensDistortionTool->GetCameraCalibrationStepsController();
+
+	if (!StepsController)
+	{
+		return;
+	}
+
+	// If not paused, cache the 3d locations of each calibrator point associated with the current calibrator
+	if (!StepsController->IsPaused())
+	{
+		// Get calibration point data
+		do
+		{
+			LastCalibratorPoints.Empty();
+
+			for (const TSharedPtr<FCalibratorPointData>& CalibratorPoint : CalibratorPoints)
+			{
+				if (!CalibratorPoint.IsValid())
+				{
+					continue;
+				}
+
+				FLensDistortionPointsCalibratorPointData PointCache;
+
+				if (!GetCalibratorPointCacheFromName(CalibratorPoint->Name, PointCache))
+				{
+					continue;
+				}
+
+				LastCalibratorPoints.Emplace(MoveTemp(PointCache));
+			}
+		} while (0);
+
+		// Get camera / calibrator data
+		do
+		{
+			LastCameraData.bIsValid = false;
+
+			const ACameraActor* Camera = StepsController->GetCamera();
+
+			if (!Camera)
+			{
+				break;
+			}
+
+			LastCameraData.UniqueId = Camera->GetUniqueID();
+
+			const FLensFileEvaluationInputs EvalInputs = StepsController->GetLensFileEvaluationInputs();
+			if (!EvalInputs.bIsValid)
+			{
+				break;
+			}
+
+			LastCameraData.InputFocus = EvalInputs.Focus;
+			LastCameraData.InputZoom = EvalInputs.Zoom;
+
+			const UCameraComponent* CameraComponent = Camera->GetCameraComponent();
+
+			if (!CameraComponent)
+			{
+				break;
+			}
+
+			LastCameraData.Pose = CameraComponent->GetComponentToWorld();
+
+			if (!Calibrator.IsValid())
+			{
+				break;
+			}
+
+			LastCameraData.CalibratorUniqueId = Calibrator->GetUniqueID();
+
+			LastCameraData.bIsValid = true;
+
+		} while (0);
+	}
+}
+
+bool UCameraLensDistortionAlgoPoints::OnViewportClicked(const FGeometry& MyGeometry, const FPointerEvent& MouseEvent)
+{
+	// We only respond to left clicks
+	if (MouseEvent.GetEffectingButton() != EKeys::LeftMouseButton)
+	{
+		return false;
+	}
+
+	if (!LensDistortionTool.IsValid())
+	{
+		return true;
+	}
+
+	FCameraCalibrationStepsController* StepsController = LensDistortionTool->GetCameraCalibrationStepsController();
+
+	if (!StepsController)
+	{
+		return true;
+	}
+
+	if (!LastCameraData.bIsValid)
+	{
+		FText ErrorMessage = LOCTEXT("InvalidLastCameraData", "Could not find a cached set of camera data (e.g. FIZ). Check the Lens Component to make sure it has valid evaluation inputs.");
+		FMessageDialog::Open(EAppMsgType::Ok, ErrorMessage);
+		return true;
+	}
+
+	// Get currently selected calibrator point
+	FLensDistortionPointsCalibratorPointData LastCalibratorPoint;
+	LastCalibratorPoint.bIsValid = false;
+	{
+		if (!CurrentCalibratorPoint.IsValid())
+		{
+			return true;
+		}
+
+		// Find its values in the cache
+		for (const FLensDistortionPointsCalibratorPointData& PointCache : LastCalibratorPoints)
+		{
+			if (PointCache.bIsValid && (PointCache.Name == CurrentCalibratorPoint->Name))
+			{
+				LastCalibratorPoint = PointCache;
+				break;
+			}
+		}
+	}
+
+	// Check that we have a valid calibrator 3dpoint or camera data
+	if (!LastCalibratorPoint.bIsValid)
+	{
+		return true;
+	}
+
+	// Get the mouse click 2d position
+	if (!StepsController->CalculateNormalizedMouseClickPosition(MyGeometry, MouseEvent, LastCalibratorPoint.Point2d, ESimulcamViewportPortion::CameraFeed))
+	{
+		return true;
+	}
+
+	// Store the 2D point in pixel coordinates
+	const FIntPoint ImageSize = StepsController->GetCameraFeedSize();
+	LastCalibratorPoint.Point2d = LastCalibratorPoint.Point2d * ImageSize;
+
+	// Export the latest session data
+	ExportSessionData();
+
+	// Create the row that we're going to add
+	TSharedPtr<FLensDistortionPointsRowData> Row = MakeShared<FLensDistortionPointsRowData>();
+
+	// Get the next row index for the current calibration session to assign to this new row
+	const uint32 RowIndex = LensDistortionTool->AdvanceSessionRowIndex();
+
+	Row->Index = RowIndex;
+	Row->CameraData = LastCameraData;
+	Row->CalibratorPointData = LastCalibratorPoint;
+	Row->PatternIndex = CurrentPatternIndex;
+
+	// Validate the new row, show a message if validation fails.
+ 	{
+ 		FText ErrorMessage;
+ 
+ 		if (!ValidateNewRow(Row, ErrorMessage))
+ 		{
+ 			FMessageDialog::Open(EAppMsgType::Ok, ErrorMessage, LOCTEXT("NewRowError", "New Row Error"));
+ 			return true;
+ 		}
+ 	}
+
+	// Add this data point
+	CalibrationRows.Add(Row);
+
+	// Export the data for this row to a .json file on disk
+	ExportRow(Row);
+
+	// Notify the ListView of the new data
+	if (CalibrationListView.IsValid())
+	{
+		CalibrationListView->RequestListRefresh();
+		CalibrationListView->RequestNavigateToItem(Row);
+	}
+
+	// Auto-advance to the next calibration point (if it exists)
+	if (AdvanceCalibratorPoint())
+	{
+		// Play media if this was the last point in the object
+		StepsController->Play();
+
+		// If the calibrator points have looped around to the first point again, advance the pattern index
+		++CurrentPatternIndex;
+	}
+
+	return true;
+}
+
+TSharedRef<SWidget> UCameraLensDistortionAlgoPoints::BuildUI()
+{
+	return SNew(SVerticalBox)
+
+		+ SVerticalBox::Slot() // Calibrator picker
+		.VAlign(EVerticalAlignment::VAlign_Top)
+		.AutoHeight()
+		.MaxHeight(FCameraCalibrationWidgetHelpers::DefaultRowHeight)
+		[FCameraCalibrationWidgetHelpers::BuildLabelWidgetPair(LOCTEXT("Calibrator", "Calibrator"), BuildCalibrationDevicePickerWidget())]
+
+		+ SVerticalBox::Slot() // Calibrator point names
+		.AutoHeight()
+		.MaxHeight(FCameraCalibrationWidgetHelpers::DefaultRowHeight)
+		[FCameraCalibrationWidgetHelpers::BuildLabelWidgetPair(LOCTEXT("CurrentCalibratorPoint", "Current Calibrator Point"), BuildCurrentCalibratorPointLabel())]
+
+		+ SVerticalBox::Slot() // Current focal length
+		.AutoHeight()
+		.MaxHeight(FCameraCalibrationWidgetHelpers::DefaultRowHeight)
+		[FCameraCalibrationWidgetHelpers::BuildLabelWidgetPair(LOCTEXT("CurrentFocalLength", "Current Focal Length (mm)"), BuildCurrentFocalLengthWidget())]
+
+		+ SVerticalBox::Slot() // Spacer
+		[
+			SNew(SBox)
+			.MinDesiredHeight(0.25 * FCameraCalibrationWidgetHelpers::DefaultRowHeight)
+			.MaxDesiredHeight(0.25 * FCameraCalibrationWidgetHelpers::DefaultRowHeight)
+		]
+
+		+ SVerticalBox::Slot() // Calibration Rows
+		.AutoHeight()
+		.MaxHeight(12 * FCameraCalibrationWidgetHelpers::DefaultRowHeight)
+		[BuildCalibrationPointsTable()]
+
+		+ SVerticalBox::Slot() // Action buttons (e.g. Remove, Clear)
+		.HAlign(EHorizontalAlignment::HAlign_Center)
+		.AutoHeight()
+		.Padding(0, 20)
+		[BuildCalibrationActionButtons()];
+}
+
+FDistortionCalibrationTask UCameraLensDistortionAlgoPoints::BeginCalibration(FText& OutErrorMessage)
+{
+	FDistortionCalibrationTask CalibrationTask = {};
+
+	// Validate that enough points were gathered to attempt a calibration
+	if (CalibrationRows.Num() < 1)
+	{
+		OutErrorMessage = LOCTEXT("NotEnoughCalibrationRowsError", "Could not initiate distortion calibration. At least 1 calibration row is required.");
+		return CalibrationTask;
+	}
+
+	if (!LensDistortionTool.IsValid())
+	{
+		OutErrorMessage = LOCTEXT("InvalidTool", "Invalid Tool");
+		return CalibrationTask;
+	}
+
+	const FCameraCalibrationStepsController* StepsController = LensDistortionTool->GetCameraCalibrationStepsController();
+
+	if (!StepsController)
+	{
+		OutErrorMessage = LOCTEXT("InvalidStepsController", "Invalid StepsController");
+		return CalibrationTask;
+	}
+
+	const ULensFile* LensFile = StepsController->GetLensFile();
+
+	if (!LensFile)
+	{
+		OutErrorMessage = LOCTEXT("InvalidLensFile", "Invalid Lens File");
+		return CalibrationTask;
+	}
+
+	const float Focus = CalibrationRows[0]->CameraData.InputFocus;
+	const float Zoom = CalibrationRows[0]->CameraData.InputZoom;
+
+	TArray<FObjectPoints> Samples3d;
+	TArray<FImagePoints> Samples2d;
+	Samples3d.Reserve(CalibrationRows.Num());
+	Samples2d.Reserve(CalibrationRows.Num());
+
+	TArray<FTransform> CameraPoses;
+
+	uint32 NextPatternIndex = 0;
+	for (int32 RowIndex = 0; RowIndex < CalibrationRows.Num(); ++RowIndex)
+	{
+		FTransform Pose = CalibrationRows[RowIndex]->CameraData.Pose;
+
+		FObjectPoints Points3d;
+		FImagePoints Points2d;
+		while (RowIndex < CalibrationRows.Num() && CalibrationRows[RowIndex]->PatternIndex == NextPatternIndex)
+		{
+			const TSharedPtr<FLensDistortionPointsRowData>& Row = CalibrationRows[RowIndex];
+
+			Points3d.Points.Add(Row->CalibratorPointData.Point3d);
+
+			const FVector2f& Point2d = Row->CalibratorPointData.Point2d;
+			Points2d.Points.Add(FVector2D(Point2d.X, Point2d.Y));
+
+			++RowIndex;
+		}
+
+		if (Points3d.Points.Num() > 0)
+		{
+			Samples3d.Add(Points3d);
+			Samples2d.Add(Points2d);
+			CameraPoses.Add(Pose);
+		}
+
+		--RowIndex;
+		++NextPatternIndex;
+	}
+
+	// Because the calibration pattern is not coplanar, OpenCV requires an initial intrinsics guess
+	if (FMath::IsNearlyEqual(FocalLengthEstimate, 0.0) || FocalLengthEstimate < 0.0)
+	{
+		OutErrorMessage = LOCTEXT("InvalidFocalLengthMsg", "The current focal length (in mm) must be set to a "
+			"valid value in order to seed the calibration algorithm with a best guess as to the camera intrinsics (Fx/Fy)");
+		return CalibrationTask;
+	}
+
+	// Validate sensor dimensions
+	const FIntPoint ImageSize = LensFile->CameraFeedInfo.GetDimensions();
+
+	const float PixelAspect = LensFile->LensInfo.SqueezeFactor;
+
+	if (FMath::IsNearlyZero(PixelAspect))
+	{
+		OutErrorMessage = LOCTEXT("PixelAspectNearlyZero", "The pixel aspect in the Lens Information is zero, which is invalid.");
+		return CalibrationTask;
+	}
+
+	const float PhysicalSensorWidth = StepsController->GetLensFileEvaluationInputs().Filmback.SensorWidth;
+	const float DesqueezeSensorWidth = PhysicalSensorWidth * PixelAspect;
+
+	if (FMath::IsNearlyZero(DesqueezeSensorWidth))
+	{
+		OutErrorMessage = LOCTEXT("SensorWidthNearlyZero", "The sensor width of the CineCamera is zero, which is invalid.");
+		return CalibrationTask;
+	}
+
+	const double Fx = (FocalLengthEstimate / DesqueezeSensorWidth) * ImageSize.X;
+
+	// When operating on a desqueezed image, we expect our pixel aspect to be square, so horizontal and vertical field of view are assumed to be equal (i.e. Fx == Fy)
+	FVector2D FocalLength = FVector2D(Fx, Fx);
+	FVector2D ImageCenter = FVector2D((ImageSize.X - 1) * 0.5, (ImageSize.Y - 1) * 0.5);
+
+	ECalibrationFlags SolverFlags = ECalibrationFlags::None;
+	EnumAddFlags(SolverFlags, ECalibrationFlags::UseIntrinsicGuess);
+
+	if (CVarUseExtrinsicsGuessPoints.GetValueOnGameThread())
+	{
+		EnumAddFlags(SolverFlags, ECalibrationFlags::UseExtrinsicGuess);
+	}
+
+	if (CVarFixExtrinsicsPoints.GetValueOnGameThread())
+	{
+		EnumAddFlags(SolverFlags, ECalibrationFlags::FixExtrinsics);
+	}
+
+	if (CVarFixDistortionPoints.GetValueOnGameThread())
+	{
+		EnumAddFlags(SolverFlags, ECalibrationFlags::FixDistortion);
+	}
+
+	FDistortionInfo DistortionGuess;
+	LensFile->EvaluateDistortionParameters(Focus, Zoom, DistortionGuess);
+
+	const TSubclassOf<ULensModel> Model = LensFile->LensInfo.LensModel;
+
+	Solver = NewObject<ULensDistortionSolver>(this, LensDistortionTool->GetSolverClass());
+
+	// When capturing individual marker positions, there is no need to calibrate for a tracking offset of the prop
+	// Therefore there is no need to capture any target poses for the calibrator actor
+	TArray<FTransform> TargetPoses;
+
+	CalibrationTask = UE::Tasks::Launch(UE_SOURCE_LOCATION, [Solver = Solver, Model, Samples3d, Samples2d, ImageSize, FocalLength, ImageCenter, DistortionGuess, CameraPoses, TargetPoses, PixelAspect, SolverFlags, Focus, Zoom]() mutable
+	{
+		FDistortionCalibrationResult Result = Solver->Solve(
+			Samples3d,
+			Samples2d,
+			ImageSize,
+			FocalLength,
+			ImageCenter,
+			DistortionGuess.Parameters,
+			CameraPoses,
+			TargetPoses,
+			Model,
+			PixelAspect,
+			SolverFlags
+		);
+
+		// CalibrateCamera() returns focal length and image center in pixels, but the result is expected to be normalized by the image size
+		Result.FocalLength.FxFy = Result.FocalLength.FxFy / ImageSize;
+		Result.ImageCenter.PrincipalPoint = Result.ImageCenter.PrincipalPoint / ImageSize;
+
+		// FZ inputs to LUT
+		Result.EvaluatedFocus = Focus;
+		Result.EvaluatedZoom = Zoom;
+
+		return Result;
+	});
+
+	return CalibrationTask;
+}
+
+void UCameraLensDistortionAlgoPoints::CancelCalibration()
+{
+	if (Solver)
+	{
+		Solver->Cancel();
+	}
+}
+
+bool UCameraLensDistortionAlgoPoints::GetCalibrationStatus(FText& StatusText) const
+{
+	if (Solver)
+	{
+		return Solver->GetStatusText(StatusText);
+	}
+	return false;
+}
+
+void UCameraLensDistortionAlgoPoints::OnDistortionSavedToLens()
+{
+	// Since the distortion info was saved, there is no further use for the current samples.
+	ClearCalibrationRows();
+}
+
+AActor* UCameraLensDistortionAlgoPoints::FindFirstCalibrator() const
+{
+	// We find the first UCalibrationPointComponent object and return its actor owner.
+	if (!LensDistortionTool.IsValid())
+	{
+		return nullptr;
+	}
+
+	const FCameraCalibrationStepsController* StepsController = LensDistortionTool->GetCameraCalibrationStepsController();
+
+	if (!StepsController)
+	{
+		return nullptr;
+	}
+
+	const UWorld* World = StepsController->GetWorld();
+	const EObjectFlags ExcludeFlags = RF_ClassDefaultObject; // We don't want the calibrator CDOs.
+
+	for (TObjectIterator<UCalibrationPointComponent> It(ExcludeFlags, true, EInternalObjectFlags::Garbage); It; ++It)
+	{
+		AActor* Actor = It->GetOwner();
+
+		if (Actor && (Actor->GetWorld() == World))
+		{
+			return Actor;
+		}
+	}
+
+	return nullptr;
+}
+
+void UCameraLensDistortionAlgoPoints::SetCalibrator(AActor* InCalibrator)
+{
+	Calibrator = InCalibrator;
+
+	// Update the list of points
+	CalibratorPoints.Empty();
+
+	if (!Calibrator.IsValid())
+	{
+		return;
+	}
+
+	TArray<UCalibrationPointComponent*, TInlineAllocator<NumInlineAllocations>> CalibrationPoints;
+	Calibrator->GetComponents(CalibrationPoints);
+
+	for (const UCalibrationPointComponent* CalibrationPoint : CalibrationPoints)
+	{
+		TArray<FString> PointNames;
+
+		CalibrationPoint->GetNamespacedPointNames(PointNames);
+
+		for (FString& PointName : PointNames)
+		{
+			CalibratorPoints.Add(MakeShared<FCalibratorPointData>(PointName));
+		}
+	}
+
+	if (CalibratorPoints.Num() > 0)
+	{
+		CurrentCalibratorPoint = CalibratorPoints[0];
+	}
+	else
+	{
+		CurrentCalibratorPoint = nullptr;
+	}
+}
+
+bool UCameraLensDistortionAlgoPoints::AdvanceCalibratorPoint()
+{
+	if (!CurrentCalibratorPoint.IsValid())
+	{
+		return false;
+	}
+
+	for (int32 PointIdx = 0; PointIdx < CalibratorPoints.Num(); PointIdx++)
+	{
+		if (CalibratorPoints[PointIdx]->Name == CurrentCalibratorPoint->Name)
+		{
+			const int32 NextIdx = (PointIdx + 1) % CalibratorPoints.Num();
+			CurrentCalibratorPoint = CalibratorPoints[NextIdx];
+
+			// return true if we wrapped around (NextIdx is zero)
+			return !NextIdx;
+		}
+	}
+
+	return false;
+}
+
+bool UCameraLensDistortionAlgoPoints::GetCalibratorPointCacheFromName(const FString& Name, FLensDistortionPointsCalibratorPointData& CalibratorPointCache) const
+{
+	CalibratorPointCache.bIsValid = false;
+
+	if (!Calibrator.IsValid())
+	{
+		return false;
+	}
+
+	TArray<UCalibrationPointComponent*, TInlineAllocator<4>> CalibrationPoints;
+	Calibrator->GetComponents(CalibrationPoints);
+
+	for (const UCalibrationPointComponent* CalibrationPoint : CalibrationPoints)
+	{
+		if (CalibrationPoint->GetWorldLocation(Name, CalibratorPointCache.Point3d))
+		{
+			CalibratorPointCache.bIsValid = true;
+			CalibratorPointCache.Name = Name;
+			return true;
+		}
+	}
+
+	return false;
+}
+
+bool UCameraLensDistortionAlgoPoints::ValidateNewRow(TSharedPtr<FLensDistortionPointsRowData>& Row, FText& OutErrorMessage) const
+{
+	const FCameraCalibrationStepsController* StepsController = LensDistortionTool->GetCameraCalibrationStepsController();
+
+	if (!StepsController)
+	{
+		OutErrorMessage = LOCTEXT("InvalidStepsController", "Invalid StepsController");
+		return false;
+	}
+
+	if (!Row.IsValid())
+	{
+		OutErrorMessage = LOCTEXT("InvalidRowPointer", "Invalid row pointer");
+		return false;
+	}
+
+	// The remaining checks all validate that data in the new row matches data in the first row.
+	// If this is the first row, those checks are unnecessary. 
+	if (CalibrationRows.Num() == 0)
+	{
+		return true;
+	}
+
+	// Same camera as before
+	const TSharedPtr<FLensDistortionPointsRowData>& FirstRow = CalibrationRows[0];
+
+	if (FirstRow->CameraData.UniqueId != Row->CameraData.UniqueId)
+	{
+		OutErrorMessage = LOCTEXT("CameraChangedDuringTheTest", "Camera changed during the test");
+		return false;
+	}
+
+	// Same calibrator as before
+	if (FirstRow->CameraData.CalibratorUniqueId != Row->CameraData.CalibratorUniqueId)
+	{
+		OutErrorMessage = LOCTEXT("CalibratorChangedDuringTheTest", "Calibrator changed during the test");
+		return false;
+	}
+
+	//@todo Focus and zoom did not change much (i.e. inputs to distortion and nodal offset). 
+	//      Threshold for physical units should differ from normalized encoders.
+
+	return true;
+}
+
+void UCameraLensDistortionAlgoPoints::ClearCalibrationRows()
+{
+	CalibrationRows.Empty();
+
+	if (CalibrationListView.IsValid())
+	{
+		CalibrationListView->RequestListRefresh();
+	}
+
+	// Reset the pattern index
+	CurrentPatternIndex = 0;
+
+	// Reset the next calibrator point to the first one
+	if (CalibratorPoints.Num() > 0)
+	{
+		CurrentCalibratorPoint = CalibratorPoints[0];
+	}
+	else
+	{
+		CurrentCalibratorPoint = nullptr;
+	}
+
+	// End the current calibration session (a new one will begin the next time a new row is added)
+	if (ULensDistortionTool* Tool = LensDistortionTool.Get())
+	{
+		Tool->EndCalibrationSession();
+	}
+}
+
+void UCameraLensDistortionAlgoPoints::DeleteSelectedCalibrationRows()
+{
+	const TArray<TSharedPtr<FLensDistortionPointsRowData>> SelectedItems = CalibrationListView->GetSelectedItems();
+
+	TSet<uint32> PatternsToDelete;
+
+	for (const TSharedPtr<FLensDistortionPointsRowData>& Row : SelectedItems)
+	{
+		PatternsToDelete.Add(Row->PatternIndex);
+	}
+
+	if (PatternsToDelete.Contains(CurrentPatternIndex))
+	{
+		if (CalibratorPoints.Num() > 0)
+		{
+			CurrentCalibratorPoint = CalibratorPoints[0];
+		}
+		else
+		{
+			CurrentCalibratorPoint = nullptr;
+		}
+	}
+
+	for (const TSharedPtr<FLensDistortionPointsRowData>& SelectedItem : SelectedItems)
+	{
+		CalibrationRows.Remove(SelectedItem);
+
+		// Delete the previously exported .json file for the row that is being deleted
+		if (ULensDistortionTool* Tool = LensDistortionTool.Get())
+		{
+			Tool->DeleteExportedRow(SelectedItem->Index);
+		}
+	}
+
+	CalibrationListView->RequestListRefresh();
+}
+
+void UCameraLensDistortionAlgoPoints::OnCalibrationRowSelectionChanged(TSharedPtr<FLensDistortionPointsRowData> SelectedRow, ESelectInfo::Type SelectInfo)
+{
+	if (SelectInfo != ESelectInfo::Type::Direct)
+	{
+		TSet<uint32> PatternsToSelect;
+
+		for (const TSharedPtr<FLensDistortionPointsRowData>& Row : CalibrationRows)
+		{
+			if (CalibrationListView->IsItemSelected(Row))
+			{
+				PatternsToSelect.Add(Row->PatternIndex);
+			}
+		}
+
+		TArray<TSharedPtr<FLensDistortionPointsRowData>> ItemsToSelect;
+
+		// Select all of the rows that have a matching pattern index (so that whole patterns are selected as groups)
+		for (const TSharedPtr<FLensDistortionPointsRowData>& Row : CalibrationRows)
+		{
+			if (PatternsToSelect.Contains(Row->PatternIndex))
+			{
+				ItemsToSelect.Add(Row);
+			}
+		}
+
+		CalibrationListView->SetItemSelection(ItemsToSelect, true);
+	}
+}
+
+TSharedRef<SWidget> UCameraLensDistortionAlgoPoints::BuildCalibrationDevicePickerWidget()
+{
+	return SNew(SFilterableActorPicker)
+		.OnSetObject_Lambda([&](const FAssetData& AssetData) -> void
+		{
+			if (AssetData.IsValid())
+			{
+				SetCalibrator(Cast<AActor>(AssetData.GetAsset()));
+			}
+		})
+		.OnShouldFilterAsset_Lambda([&](const FAssetData& AssetData) -> bool
+		{
+			const AActor* Actor = Cast<AActor>(AssetData.GetAsset());
+
+			if (!Actor)
+			{
+				return false;
+			}
+
+			TArray<UCalibrationPointComponent*, TInlineAllocator<NumInlineAllocations>> CalibrationPoints;
+			Actor->GetComponents(CalibrationPoints);
+
+			return (CalibrationPoints.Num() > 0);
+		})
+		.ActorAssetData_Lambda([&]() -> FAssetData
+		{
+			return FAssetData(Calibrator.Get(), true);
+		});
+}
+
+TSharedRef<SWidget> UCameraLensDistortionAlgoPoints::BuildCurrentCalibratorPointLabel()
+{
+	return SNew(STextBlock)
+		.Text_Lambda([&]() -> FText
+		{
+			if (CurrentCalibratorPoint.IsValid())
+			{
+				return FText::FromString(CurrentCalibratorPoint->Name);
+			}
+
+			return LOCTEXT("NoCurrentCalibrator", "None");
+		})
+		.ColorAndOpacity(FLinearColor::White);
+}
+
+TSharedRef<SWidget> UCameraLensDistortionAlgoPoints::BuildCurrentFocalLengthWidget()
+{
+	return SNew(SNumericEntryBox<double>)
+		.Value_Lambda([&]() { return FocalLengthEstimate; })
+		.ToolTipText(LOCTEXT("CurrentFocalLength", "Current Focal Length (mm)"))
+		.OnValueChanged_Lambda([&](double InValue)
+		{
+			if (InValue < 1.0f)
+			{
+				InValue = 1.0f;
+			}
+			FocalLengthEstimate = InValue;
+		})
+		.OnValueCommitted_Lambda([&](double InValue, ETextCommit::Type CommitType)
+		{
+			// Re-export the session data to update the new focal length value
+			ExportSessionData();
+		});
+}
+
+TSharedRef<SWidget> UCameraLensDistortionAlgoPoints::BuildCalibrationPointsTable()
+{
+	CalibrationListView = SNew(SListView<TSharedPtr<FLensDistortionPointsRowData>>)
+		.ListItemsSource(&CalibrationRows)
+		.OnGenerateRow_Lambda([&](TSharedPtr<FLensDistortionPointsRowData> InItem, const TSharedRef<STableViewBase>& OwnerTable) -> TSharedRef<ITableRow>
+		{
+			return SNew(CameraLensDistortionAlgoPoints::SCalibrationRowGenerator, OwnerTable).CalibrationRowData(InItem);
+		})
+		.SelectionMode(ESelectionMode::Multi)
+		.OnSelectionChanged_Lambda([&](TSharedPtr<FLensDistortionPointsRowData> SelectedRow, ESelectInfo::Type SelectInfo)
+		{
+			OnCalibrationRowSelectionChanged(SelectedRow, SelectInfo);
+		})
+		.OnKeyDownHandler_Lambda([&](const FGeometry& Geometry, const FKeyEvent& KeyEvent) -> FReply
+		{
+			if (!CalibrationListView.IsValid())
+			{
+				return FReply::Unhandled();
+			}
+
+			if (KeyEvent.GetKey() == EKeys::Delete)
+			{
+				DeleteSelectedCalibrationRows();
+				return FReply::Handled();
+			}
+			else if (KeyEvent.GetModifierKeys().IsControlDown() && (KeyEvent.GetKey() == EKeys::A))
+			{
+				// Select all items
+				CalibrationListView->SetItemSelection(CalibrationRows, true);
+				return FReply::Handled();
+			}
+			else if (KeyEvent.GetKey() == EKeys::Escape)
+			{
+				// Deselect all items
+				CalibrationListView->ClearSelection();
+				return FReply::Handled();
+			}
+
+			return FReply::Unhandled();
+		})
+		.HeaderRow
+		(
+			SNew(SHeaderRow)
+
+			+ SHeaderRow::Column("Pattern")
+			.DefaultLabel(LOCTEXT("Pattern", "Pattern"))
+			.FillWidth(0.15f)
+
+			+ SHeaderRow::Column("Name")
+			.DefaultLabel(LOCTEXT("Name", "Name"))
+			.FillWidth(0.25f)
+
+			+ SHeaderRow::Column("Point2D")
+			.DefaultLabel(LOCTEXT("Point2D", "Point2D"))
+			.FillWidth(0.3f)
+
+			+ SHeaderRow::Column("Point3D")
+			.DefaultLabel(LOCTEXT("Point3D", "Point3D"))
+			.FillWidth(0.3f)
+		);
+
+	return CalibrationListView.ToSharedRef();
+}
+
+TSharedRef<SWidget> UCameraLensDistortionAlgoPoints::BuildCalibrationActionButtons()
+{
+	return SNew(SHorizontalBox)
+
+		+ SHorizontalBox::Slot() // Button to clear all rows
+		.AutoWidth()
+		[
+			SNew(SButton)
+			.Text(LOCTEXT("ClearAll", "Clear All"))
+			.HAlign(HAlign_Center)
+			.VAlign(VAlign_Center)
+			.OnClicked_Lambda([&]() -> FReply
+			{
+				ClearCalibrationRows();
+				return FReply::Handled();
+			})
+		]
+	;
+}
+
+void UCameraLensDistortionAlgoPoints::ExportSessionData()
+{
+	using namespace UE::CameraCalibration::Private;
+	if (ULensDistortionTool* Tool = LensDistortionTool.Get())
+	{
+		// Add all data to a json object that is needed to run this algorithm that is NOT part of a specific row
+		TSharedPtr<FJsonObject> JsonSessionData = MakeShared<FJsonObject>();
+
+		JsonSessionData->SetNumberField(LensDistortionPointsExportFields::CurrentFocalLength, FocalLengthEstimate);
+		JsonSessionData->SetNumberField(LensDistortionPointsExportFields::Version, DATASET_VERSION);
+
+		// Export the session data to a .json file
+		Tool->ExportSessionData(JsonSessionData.ToSharedRef());
+	}
+}
+
+void UCameraLensDistortionAlgoPoints::ExportRow(TSharedPtr<FLensDistortionPointsRowData> Row)
+{
+	if (ULensDistortionTool* Tool = LensDistortionTool.Get())
+	{
+		if (const TSharedPtr<FJsonObject>& RowObject = FJsonObjectConverter::UStructToJsonObject<FLensDistortionPointsRowData>(Row.ToSharedRef().Get()))
+		{
+			// Export the row data to a .json file
+			LensDistortionTool->ExportCalibrationRow(Row->Index, RowObject.ToSharedRef());
+		}
+	}
+}
+
+bool UCameraLensDistortionAlgoPoints::HasCalibrationData() const
+{
+	return (CalibrationRows.Num() > 0);
+}
+
+void UCameraLensDistortionAlgoPoints::PreImportCalibrationData()
+{
+	// Clear the current set of rows before importing new ones
+	CalibrationRows.Empty();
+}
+
+void UCameraLensDistortionAlgoPoints::ImportSessionData(const TSharedRef<FJsonObject>& SessionDataObject)
+{
+	using namespace UE::CameraCalibration::Private;
+	SessionDataObject->TryGetNumberField(LensDistortionPointsExportFields::CurrentFocalLength, FocalLengthEstimate);
+}
+
+int32 UCameraLensDistortionAlgoPoints::ImportCalibrationRow(const TSharedRef<FJsonObject>& CalibrationRowObject, const FImage& RowImage)
+{
+	// Create a new row to populate with data from the Json object
+	TSharedPtr<FLensDistortionPointsRowData> NewRow = MakeShared<FLensDistortionPointsRowData>();
+
+	// We enforce strict mode to ensure that every field in the UStruct of row data is present in the imported json.
+	// If any fields are missing, it is likely the row will be invalid, which will lead to errors in the calibration.
+	constexpr int64 CheckFlags = 0;
+	constexpr int64 SkipFlags = 0;
+	constexpr bool bStrictMode = true;
+	if (FJsonObjectConverter::JsonObjectToUStruct<FLensDistortionPointsRowData>(CalibrationRowObject, NewRow.Get(), CheckFlags, SkipFlags, bStrictMode))
+	{
+		CalibrationRows.Add(NewRow);
+	}
+	else
+	{
+		UE_LOG(LogCameraCalibrationEditor, Warning, TEXT("LensDistortionPoints algo failed to import calibration row because at least one field could not be deserialized from the json file."));
+	}
+
+	return NewRow->Index;
+}
+
+void UCameraLensDistortionAlgoPoints::PostImportCalibrationData()
+{
+	// Find the maximum pattern index among the imported rows and set the current index to be greater than that
+	uint32 MaxPatternIndex = 0;
+	for (const TSharedPtr<FLensDistortionPointsRowData>& Row : CalibrationRows)
+	{
+		MaxPatternIndex = FMath::Max(MaxPatternIndex, Row->PatternIndex);
+	}
+
+	CurrentPatternIndex = MaxPatternIndex + 1;
+
+	// Sort imported calibration rows by row index
+	CalibrationRows.Sort([](const TSharedPtr<FLensDistortionPointsRowData>& LHS, const TSharedPtr<FLensDistortionPointsRowData>& RHS) { return LHS->Index < RHS->Index; });
+
+	// Notify the ListView of the new data
+	if (CalibrationListView.IsValid())
+	{
+		CalibrationListView->RequestListRefresh();
+	}
+}
+
+TSharedRef<SWidget> UCameraLensDistortionAlgoPoints::BuildHelpWidget()
+{
+	return SNew(STextBlock)
+		.Text(LOCTEXT("LensDistortionPointsHelp",
+			"This Lens Distortion algorithm is based on capturing at least 4 different views of a calibration pattern.\n\n"
+			"The camera and/or the calibrator may be moved for each. To capture, simply click the simulcam\n"
+			"viewport. You can optionally right-click the simulcam viewport to pause it and ensure it will be\n"
+			"a sharp capture (if the media source supports pause).\n\n"
+
+			"This will require the selection of a calibrator actor that exists in the scene, and is configured\n"
+			"to have several calibration points components that correspond to physical features that are easy to detect.\n\n"
+
+			"Each click will add a row to the table with the 3d and 2d points of the current calibration point. The tool\n"
+			"will validate that you capture at least 4 patterns, and that each pattern has the same number of calibration\n"
+			"points in it. It is important that captures sweep the viewport area as much as possible, so that the\n"
+			"calibration includes samples from all of its regions (otherwise the calibration may be skewed towards a\n"
+			"particular region and not be very accurate in others).\n\n"
+
+			"When you are done capturing, click the 'Add To Lens Distortion Calibration' button to run the calibration\n"
+			"algorithm, which calculates the distortion parameters k1,k2,p1,p2,k3 and adds them to the lens file.\n\n"
+
+			"A dialog window will show you the reprojection error of the calibration, and you can decide to add the\n"
+			"data to the lens file or not.\n"
+		));
+
+}
+
+#undef LOCTEXT_NAMESPACE
diff --git a/Engine/Plugins/VirtualProduction/CameraCalibration/Source/CameraCalibrationEditor/Private/Calibrators/CameraLensDistortionAlgoPoints.h b/Engine/Plugins/VirtualProduction/CameraCalibration/Source/CameraCalibrationEditor/Private/Calibrators/CameraLensDistortionAlgoPoints.h
new file mode 100644
--- /dev/null
+++ b/Engine/Plugins/VirtualProduction/CameraCalibration/Source/CameraCalibrationEditor/Private/Calibrators/CameraLensDistortionAlgoPoints.h
@@ -0,0 +1,249 @@
+// Copyright Epic Games, Inc. All Rights Reserved.
+
+#pragma once
+
+#include "CameraLensDistortionAlgo.h"
+
+#include "CameraCalibrationSolver.h"
+#include "LensFile.h"
+
+#include "CameraLensDistortionAlgoPoints.generated.h"
+
+template <typename ItemType>
+class SListView;
+
+class FJsonObject;
+
+namespace CameraLensDistortionAlgoPoints
+{
+	class SCalibrationRowGenerator;
+};
+
+/** Holds information of the identified calibrator 2d point for a given sample of a 2d-3d correlation */
+USTRUCT()
+struct FLensDistortionPointsCalibratorPointData
+{
+	GENERATED_BODY()
+
+	// True if the rest of the contents are valid.
+	UPROPERTY()
+	bool bIsValid = false;
+
+	// The name of the 3d calibrator point, as defined in the mesh
+	UPROPERTY()
+	FString Name;
+
+	// 3D location of the calibrator point
+	UPROPERTY()
+	FVector Point3d = FVector::ZeroVector;
+
+	// 2D location of the calibrator point
+	UPROPERTY()
+	FVector2f Point2d = FVector2f::ZeroVector;
+};
+
+/** Holds information of the camera pose for a given sample of a 2d-3d correlation */
+USTRUCT()
+struct FLensDistortionPointsCameraData
+{
+	GENERATED_BODY()
+
+	// True if the rest of the contents are valid.
+	UPROPERTY()
+	bool bIsValid = false;
+
+	// The unique id of the camera object. Used to detect camera selection changes during a calibration session.
+	UPROPERTY()
+	uint32 UniqueId = INDEX_NONE;
+
+	// The unique id of the calibrator object. Used to detect calibrator selection changes during a calibration session.
+	UPROPERTY()
+	uint32 CalibratorUniqueId = INDEX_NONE;
+
+	// The camera pose
+	UPROPERTY()
+	FTransform Pose = FTransform::Identity;
+
+	// Input focus from lens file evaluation data
+	UPROPERTY()
+	float InputFocus = 0.0f;
+
+	// Input zoom from lens file evaluation data
+	UPROPERTY()
+	float InputZoom = 0.0f;
+};
+
+/** Holds information of the calibrator 3d point for a given sample of a 2d-3d correlation */
+USTRUCT()
+struct FLensDistortionPointsRowData
+{
+	GENERATED_BODY()
+	
+	// Index to display in list
+	UPROPERTY()
+	int32 Index = INDEX_NONE;
+
+	// The 3D-2D point correspondences for each of the calibrator points in the chosen calibrator for a single pose
+	UPROPERTY()
+	FLensDistortionPointsCalibratorPointData CalibratorPointData;
+
+	// Holds information of the camera data for this sample
+	UPROPERTY()
+	FLensDistortionPointsCameraData CameraData;
+
+	// Index of the calibration pattern this row is associated with
+	UPROPERTY()
+	uint32 PatternIndex = INDEX_NONE;
+};
+
+/**
+ * Implements a lens distortion calibration algorithm. It uses 3d points (UCalibrationPointComponent)
+ * specified in a selectable calibrator with features that the user can identify by clicking on the
+ * simulcam viewport, and after enough points have been identified, it can calculate the lens distortion
+ * and camera intrinsics that need to be applied to the associated camera in order to align its CG with 
+ * the live action media plate.
+ *
+ * Limitations:
+ * - Only supports Brown–Conrady model lens model (FSphericalDistortionParameters)
+ * - The camera or the lens should not be moved during the calibration of each lens distortion sample.
+ */
+UCLASS()
+class UCameraLensDistortionAlgoPoints : public UCameraLensDistortionAlgo
+{
+	GENERATED_BODY()
+
+public:
+	//~ Begin UCameraLensDistortionAlgo
+	virtual void Initialize(ULensDistortionTool* InLensDistortionTool) override;
+	virtual void Shutdown() override;
+	virtual bool SupportsModel(const TSubclassOf<ULensModel>& LensModel) const override;
+	virtual void Tick(float DeltaTime) override;
+	virtual bool OnViewportClicked(const FGeometry& MyGeometry, const FPointerEvent& MouseEvent) override;
+	virtual TSharedRef<SWidget> BuildUI() override;
+	virtual FName FriendlyName() const override { return TEXT("Lens Distortion Points Method"); };
+	virtual FName ShortName() const override { return TEXT("Points"); };
+	virtual void OnDistortionSavedToLens() override;
+
+	virtual FDistortionCalibrationTask BeginCalibration(FText& OutErrorMessage) override;
+	virtual void CancelCalibration() override;
+	virtual bool GetCalibrationStatus(FText& StatusText) const override;
+	virtual bool SupportsAsyncCalibration() override { return true; };
+
+	virtual bool HasCalibrationData() const override;
+	virtual void PreImportCalibrationData() override;
+	virtual void ImportSessionData(const TSharedRef<FJsonObject>& SessionDataObject) override;
+	virtual int32 ImportCalibrationRow(const TSharedRef<FJsonObject>& CalibrationRowObject, const FImage& RowImage) override;
+	virtual void PostImportCalibrationData() override;
+	virtual TSharedRef<SWidget> BuildHelpWidget() override;
+	//~ End UCameraLensDistortionAlgo
+
+protected:
+
+	// SCalibrationRowGenerator will need access to the row structures below.
+	friend class CameraLensDistortionAlgoPoints::SCalibrationRowGenerator;
+
+	/** Holds item information a given calibrator point in the calibrator model */
+	struct FCalibratorPointData
+	{
+		FCalibratorPointData(FString& InName)
+			: Name(InName)
+		{}
+
+		/** Name of the calibrator 3d point, as defined in the mesh */
+		FString Name;
+	};
+
+protected:
+	/** Version of the current calibration dataset */
+	static const int DATASET_VERSION;
+
+	/** The lens distortion tool */
+	TWeakObjectPtr<ULensDistortionTool> LensDistortionTool;
+
+	/** The currently selected calibrator object. It is expected to contain one or more UCalibrationPointComponent in it */
+	TWeakObjectPtr<AActor> Calibrator;
+
+	/** Lists the calibrator points found in the selected calibrator object */
+	TArray<TSharedPtr<FCalibratorPointData>> CalibratorPoints;
+
+	/** The current calibrator point that will be used to generate the next calibration row */
+	TSharedPtr<FCalibratorPointData> CurrentCalibratorPoint;
+
+	/** Rows source for the CalibrationListView */
+	TArray<TSharedPtr<FLensDistortionPointsRowData>> CalibrationRows;
+
+	/** Displays the list of calibration points that will be used to calculate the lens distortion */
+	TSharedPtr<SListView<TSharedPtr<FLensDistortionPointsRowData>>> CalibrationListView;
+
+	/** Caches the last calibrator point 3d location.  Will hold last value before the lens distortion tool is paused */
+	TArray<FLensDistortionPointsCalibratorPointData> LastCalibratorPoints;
+
+	/** Caches the last camera data.  Will hold last value before the lens distortion tool is paused */
+	FLensDistortionPointsCameraData LastCameraData;
+
+	/** Instructs the inline allocation policy how many elements to allocate in a single allocation */
+	static constexpr uint32 NumInlineAllocations = 32;
+
+	/** The current index of the calibration pattern being captured, used for grouping calibration rows together into meaningful patterns */
+	uint32 CurrentPatternIndex = 0;
+
+	/** The value of the current focal length of the lens (in mm), used to set a guess for the camera intrinsics fx/fy */
+	double FocalLengthEstimate = 0.0;
+
+	/** Solver instance that will run the distortion calibration */
+	TObjectPtr<ULensDistortionSolver> Solver;
+
+protected:
+
+	/** Returns the first calibrator object in the scene that it can find */
+	AActor* FindFirstCalibrator() const;
+
+	/** Sets the calibrator object to be used. Updates the selection picker. */
+	void SetCalibrator(AActor* InCalibrator);
+
+	/** Selects the next available UCalibrationPointComponent of the currently selected calibrator object. Returns true when it wraps around */
+	bool AdvanceCalibratorPoint();
+
+	/**
+	 * Retrieves by name the calibration point data of the currently selected calibrator.
+	 *
+	 * @param Name The name of the point (namespaced if it is a subpoint).
+	 * @param CalibratorPointCache This data structure will be populated with the information found.
+	 *
+	 * @return True if successful.
+	 */
+	bool GetCalibratorPointCacheFromName(const FString& Name, FLensDistortionPointsCalibratorPointData& CalibratorPointCache) const;
+
+	/** Validates a new calibration point to determine if it should be added as a new sample row */
+	bool ValidateNewRow(TSharedPtr<FLensDistortionPointsRowData>& Row, FText& OutErrorMessage) const;
+
+	/** Clears the list of calibration sample rows */
+	void ClearCalibrationRows();
+
+	/** Deletes the selected calibration rows */
+	void DeleteSelectedCalibrationRows();
+
+	/** Updates the selection in the list of calibration rows so that patterns are always grouped together */
+	void OnCalibrationRowSelectionChanged(TSharedPtr<FLensDistortionPointsRowData> SelectedRow, ESelectInfo::Type SelectInfo);
+
+	/** Export global session data to a .json file */
+	void ExportSessionData();
+
+	/** Export the row data to a .json file */
+	void ExportRow(TSharedPtr<FLensDistortionPointsRowData> Row);
+
+	/** Builds the UI of the calibration device picker */
+	TSharedRef<SWidget> BuildCalibrationDevicePickerWidget();
+
+	/** Buils the UI for selecting the current calibrator point */
+	TSharedRef<SWidget> BuildCurrentCalibratorPointLabel();
+
+	/** Buils the UI for setting the current focal length (in mm) */
+	TSharedRef<SWidget> BuildCurrentFocalLengthWidget();
+
+	/** Builds the UI of the calibration points table */
+	TSharedRef<SWidget> BuildCalibrationPointsTable();
+
+	/** Builds the UI for the action buttons (RemoveLast, ClearAll) */
+	TSharedRef<SWidget> BuildCalibrationActionButtons();
+};
diff --git a/Engine/Plugins/VirtualProduction/CameraCalibration/Source/CameraCalibrationEditor/Private/CameraCalibrationUtilsPrivate.cpp b/Engine/Plugins/VirtualProduction/CameraCalibration/Source/CameraCalibrationEditor/Private/CameraCalibrationUtilsPrivate.cpp
--- a/Engine/Plugins/VirtualProduction/CameraCalibration/Source/CameraCalibrationEditor/Private/CameraCalibrationUtilsPrivate.cpp
+++ b/Engine/Plugins/VirtualProduction/CameraCalibration/Source/CameraCalibrationEditor/Private/CameraCalibrationUtilsPrivate.cpp
@@ -130,7 +130,7 @@ namespace UE::CameraCalibration::Private
 		}
 	}
 
-	bool FindArucoCalibrationPoint(const TArray<TWeakObjectPtr<UCalibrationPointComponent>>& CalibrationComponents, EArucoDictionary ArucoDictionary, const FArucoMarker& ArucoMarker, FArucoCalibrationPoint& OutArucoCalibrationPoint)
+	bool FindArucoCalibrationPoint(const TArray<UCalibrationPointComponent*>& CalibrationComponents, EArucoDictionary ArucoDictionary, const FArucoMarker& ArucoMarker, FArucoCalibrationPoint& OutArucoCalibrationPoint)
 	{
 		// Build calibrator point name based on the detected marker
 		const FString DictionaryName = GetArucoDictionaryName(ArucoDictionary);
@@ -141,9 +141,9 @@ namespace UE::CameraCalibration::Private
 		OutArucoCalibrationPoint.MarkerID = ArucoMarker.MarkerID;
 		OutArucoCalibrationPoint.Name = FString::Printf(TEXT("%s-%d"), *DictionaryName, ArucoMarker.MarkerID);
 
-		for (TWeakObjectPtr<UCalibrationPointComponent> Component : CalibrationComponents)
+		for (UCalibrationPointComponent* Component : CalibrationComponents)
 		{
-			if (!Component.IsValid())
+			if (!Component)
 			{
 				continue;
 			}
diff --git a/Engine/Plugins/VirtualProduction/CameraCalibration/Source/CameraCalibrationEditor/Private/CameraCalibrationUtilsPrivate.h b/Engine/Plugins/VirtualProduction/CameraCalibration/Source/CameraCalibrationEditor/Private/CameraCalibrationUtilsPrivate.h
--- a/Engine/Plugins/VirtualProduction/CameraCalibration/Source/CameraCalibrationEditor/Private/CameraCalibrationUtilsPrivate.h
+++ b/Engine/Plugins/VirtualProduction/CameraCalibration/Source/CameraCalibrationEditor/Private/CameraCalibrationUtilsPrivate.h
@@ -2,7 +2,6 @@
 #pragma once
 
 #include "OpenCVHelper.h"
-#include "UObject/WeakObjectPtrTemplates.h"
 
 #include "CameraCalibrationUtilsPrivate.generated.h"
 
@@ -49,7 +48,7 @@ namespace UE::CameraCalibration::Private
 	void FindActorsWithCalibrationComponents(TArray<AActor*>& ActorsWithCalibrationComponents);
 
 	/** Find an aruco marker calibration point in one of the input calibration components that matches the input dictionary and marker ID */
-	bool FindArucoCalibrationPoint(const TArray<TWeakObjectPtr<UCalibrationPointComponent>>& CalibrationComponents, EArucoDictionary ArucoDictionary, const FArucoMarker& ArucoMarker, FArucoCalibrationPoint& OutArucoCalibrationPoint);
+	bool FindArucoCalibrationPoint(const TArray<UCalibrationPointComponent*>& CalibrationComponents, EArucoDictionary ArucoDictionary, const FArucoMarker& ArucoMarker, FArucoCalibrationPoint& OutArucoCalibrationPoint);
 
 	/** Set every pixel in the input texture to the clear color */
 	void ClearTexture(UTexture2D* Texture, FColor ClearColor = FColor::Transparent);
