CameraCalibration: Fix lens calibration tool to properly import old datasets from before the algos got combined.

diff --git a/Engine/Plugins/VirtualProduction/CameraCalibration/Source/CameraCalibrationEditor/Private/AssetEditor/LensDistortionTool.cpp b/Engine/Plugins/VirtualProduction/CameraCalibration/Source/CameraCalibrationEditor/Private/AssetEditor/LensDistortionTool.cpp
--- a/Engine/Plugins/VirtualProduction/CameraCalibration/Source/CameraCalibrationEditor/Private/AssetEditor/LensDistortionTool.cpp
+++ b/Engine/Plugins/VirtualProduction/CameraCalibration/Source/CameraCalibrationEditor/Private/AssetEditor/LensDistortionTool.cpp
@@ -949,6 +949,7 @@ void ULensDistortionTool::ImportCalibrationDataset()
 
 	// Import the session data
 	FDateTime ImportedSessionDateTime = FDateTime::Now();
+	EDatasetVersion ImportedDatasetVersion = EDatasetVersion::Invalid;
 	{
 		const FString SessionFile = SelectedDirectory / SessionFileName;
 
@@ -971,6 +972,33 @@ void ULensDistortionTool::ImportCalibrationDataset()
 				{
 					UE_LOG(LogCameraCalibrationEditor, Verbose, TEXT("Lens Distortion Tool failed to deserialize the date/time from the session file: %s"), *SessionFile);
 				}
+
+				int32 Version = 0;
+				if (JsonSessionData->TryGetNumberField(LensDistortionTool::Version, Version))
+				{
+					ImportedDatasetVersion = (EDatasetVersion)Version;
+				}
+
+				// This version predates the calibration pattern field, but the data can be reconstructed by looking at the "Algo Name" field.
+				if (ImportedDatasetVersion == EDatasetVersion::SeparateAlgoClasses)
+				{
+					FString AlgoString;
+					if (JsonSessionData->TryGetStringField(TEXT("AlgoName"), AlgoString))
+					{
+						if (AlgoString.Contains(TEXT("Checkerboard"), ESearchCase::IgnoreCase))
+						{
+							CaptureSettings.CalibrationPattern = ECalibrationPattern::Checkerboard;
+						}
+						else if (AlgoString.Contains(TEXT("Aruco"), ESearchCase::IgnoreCase))
+						{
+							CaptureSettings.CalibrationPattern = ECalibrationPattern::Aruco;
+						}
+						else if (AlgoString.Contains(TEXT("Points"), ESearchCase::IgnoreCase))
+						{
+							CaptureSettings.CalibrationPattern = ECalibrationPattern::Points;
+						}
+					}
+				}
 			}
 			else
 			{
@@ -1005,7 +1033,7 @@ void ULensDistortionTool::ImportCalibrationDataset()
 			TSharedPtr<FJsonObject> JsonRowData = MakeShared<FJsonObject>();
 			if (FJsonSerializer::Deserialize(JsonReader, JsonRowData))
 			{
-				int32 RowNum = ImportCalibrationRow(JsonRowData.ToSharedRef(), RowImage);
+				int32 RowNum = ImportCalibrationRow(JsonRowData.ToSharedRef(), RowImage, ImportedDatasetVersion);
 				MaxRowIndex = FMath::Max(MaxRowIndex, RowNum);
 			}
 			else
@@ -1040,7 +1068,9 @@ void ULensDistortionTool::ExportSessionData()
 	using namespace UE::CameraCalibration::Private;
 
 	TSharedPtr<FJsonObject> JsonSessionData = MakeShared<FJsonObject>();
-	JsonSessionData->SetNumberField(LensDistortionTool::Version, 2);
+
+	int32 DatasetVersion = (int32)EDatasetVersion::CurrentVersion;
+	JsonSessionData->SetNumberField(LensDistortionTool::Version, DatasetVersion);
 
 	// Start a calibration session (if one is not currently active)
 	StartCalibrationSession();
@@ -1114,28 +1144,197 @@ void ULensDistortionTool::ExportCalibrationRow(TSharedPtr<FCalibrationRow> Row)
 	}
 }
 
-int32 ULensDistortionTool::ImportCalibrationRow(const TSharedRef<FJsonObject>& CalibrationRowObject, const FImage& RowImage)
+int32 ULensDistortionTool::ImportCalibrationRow(const TSharedRef<FJsonObject>& CalibrationRowObject, const FImage& RowImage, EDatasetVersion DatasetVersion)
 {
 	// Create a new row to populate with data from the Json object
 	TSharedPtr<FCalibrationRow> NewRow = MakeShared<FCalibrationRow>();
 
-	// We enforce strict mode to ensure that every field in the UStruct of row data is present in the imported json.
-	// If any fields are missing, it is likely the row will be invalid, which will lead to errors in the calibration.
-	constexpr int64 CheckFlags = 0;
-	constexpr int64 SkipFlags = 0;
-	constexpr bool bStrictMode = true;
-	if (FJsonObjectConverter::JsonObjectToUStruct<FCalibrationRow>(CalibrationRowObject, NewRow.Get(), CheckFlags, SkipFlags, bStrictMode))
+	if (!RowImage.RawData.IsEmpty())
+	{
+		NewRow->MediaImage = RowImage;
+	}
+
+	if (DatasetVersion == EDatasetVersion::SeparateAlgoClasses)
 	{
-		if (!RowImage.RawData.IsEmpty())
+		NewRow->Pattern = CaptureSettings.CalibrationPattern;
+
+		CalibrationRowObject->TryGetNumberField(TEXT("index"), NewRow->Index);
+
+		if (NewRow->Pattern == ECalibrationPattern::Checkerboard)
+		{
+			const TArray<TSharedPtr<FJsonValue>>* Points3DArray;
+			if (CalibrationRowObject->TryGetArrayField(TEXT("points3d"), Points3DArray))
+			{
+				for (TSharedPtr<FJsonValue> PointValue : *Points3DArray)
+				{
+					TSharedPtr<FJsonObject> PointObject = PointValue->AsObject();
+					FVector NewPoint;
+					PointObject->TryGetNumberField(TEXT("x"), NewPoint.X);
+					PointObject->TryGetNumberField(TEXT("y"), NewPoint.Y);
+					PointObject->TryGetNumberField(TEXT("z"), NewPoint.Z);
+
+					NewRow->ObjectPoints.Points.Add(NewPoint);
+				}
+			}
+
+			const TArray<TSharedPtr<FJsonValue>>* Points2DArray;
+			if (CalibrationRowObject->TryGetArrayField(TEXT("points2d"), Points2DArray))
+			{
+				for (TSharedPtr<FJsonValue> PointValue : *Points2DArray)
+				{
+					TSharedPtr<FJsonObject> PointObject = PointValue->AsObject();
+					FVector2D NewPoint;
+					PointObject->TryGetNumberField(TEXT("x"), NewPoint.X);
+					PointObject->TryGetNumberField(TEXT("y"), NewPoint.Y);
+
+					NewRow->ImagePoints.Points.Add(NewPoint);
+				}
+			}
+
+			CalibrationRowObject->TryGetNumberField(TEXT("numCornerCols"), NewRow->CheckerboardDimensions.X);
+			CalibrationRowObject->TryGetNumberField(TEXT("numCornerRows"), NewRow->CheckerboardDimensions.Y);
+
+			if (!NewRow->ObjectPoints.Points.IsEmpty() && NewRow->ObjectPoints.Points.Num() == (NewRow->CheckerboardDimensions.X * NewRow->CheckerboardDimensions.Y))
+			{
+				FVector TopLeft = NewRow->ObjectPoints.Points[0];
+				FVector TopRight = NewRow->ObjectPoints.Points[NewRow->CheckerboardDimensions.Y - 1];
+				FVector BottomLeft = NewRow->ObjectPoints.Points[NewRow->CheckerboardDimensions.X * (NewRow->CheckerboardDimensions.Y - 1)];
+
+				NewRow->TargetPose.SetLocation(TopLeft);
+
+				FRotator BoardRotation = FRotationMatrix::MakeFromYZ(TopRight - TopLeft, TopLeft - BottomLeft).Rotator();
+				NewRow->TargetPose.SetRotation(BoardRotation.Quaternion());
+			}
+		}
+		else if (NewRow->Pattern == ECalibrationPattern::Aruco)
 		{
-			NewRow->MediaImage = RowImage;
+			const TArray<TSharedPtr<FJsonValue>>* ArucoPointArray;
+			if (CalibrationRowObject->TryGetArrayField(TEXT("arucoPoints"), ArucoPointArray))
+			{
+				for (TSharedPtr<FJsonValue> ArucoPointValue : *ArucoPointArray)
+				{
+					TSharedPtr<FJsonObject> ArucoPointObject = ArucoPointValue->AsObject();
+
+					const TArray<TSharedPtr<FJsonValue>>* Corners3DArray;
+					if (ArucoPointObject->TryGetArrayField(TEXT("corners3D"), Corners3DArray))
+					{
+						for (TSharedPtr<FJsonValue> CornerValue : *Corners3DArray)
+						{
+							TSharedPtr<FJsonObject> CornerObject = CornerValue->AsObject();
+
+							FVector Corner;
+							CornerObject->TryGetNumberField(TEXT("x"), Corner.X);
+							CornerObject->TryGetNumberField(TEXT("y"), Corner.Y);
+							CornerObject->TryGetNumberField(TEXT("z"), Corner.Z);
+
+							NewRow->ObjectPoints.Points.Add(Corner);
+						}
+					}
+
+					const TArray<TSharedPtr<FJsonValue>>* Corners2DArray;
+					if (ArucoPointObject->TryGetArrayField(TEXT("corners2D"), Corners2DArray))
+					{
+						for (TSharedPtr<FJsonValue> CornerValue : *Corners2DArray)
+						{
+							TSharedPtr<FJsonObject> CornerObject = CornerValue->AsObject();
+
+							FVector2D Corner;
+							CornerObject->TryGetNumberField(TEXT("x"), Corner.X);
+							CornerObject->TryGetNumberField(TEXT("y"), Corner.Y);
+
+							NewRow->ImagePoints.Points.Add(Corner);
+						}
+					}
+
+					if (NewRow->ObjectPoints.Points.Num() >= 4)
+					{
+						FVector TopLeft = NewRow->ObjectPoints.Points[0];
+						FVector TopRight = NewRow->ObjectPoints.Points[1];
+						FVector BottomLeft = NewRow->ObjectPoints.Points[3];
+
+						NewRow->TargetPose.SetLocation(TopLeft);
+
+						FRotator FirstMarkerRotation = FRotationMatrix::MakeFromYZ(TopRight - TopLeft, TopLeft - BottomLeft).Rotator();
+						NewRow->TargetPose.SetRotation(FirstMarkerRotation.Quaternion());
+					}
+				}
+			}
+		}
+		else if (NewRow->Pattern == ECalibrationPattern::Points)
+		{
+			const TSharedPtr<FJsonObject>* PointDataJsonObject;
+			if (CalibrationRowObject->TryGetObjectField(TEXT("calibratorPointData"), PointDataJsonObject))
+			{
+				const TSharedPtr<FJsonObject>* Point3DJsonObject;
+				if (PointDataJsonObject->Get()->TryGetObjectField(TEXT("point3d"), Point3DJsonObject))
+				{
+					FVector NewPoint;
+					Point3DJsonObject->Get()->TryGetNumberField(TEXT("x"), NewPoint.X);
+					Point3DJsonObject->Get()->TryGetNumberField(TEXT("y"), NewPoint.Y);
+					Point3DJsonObject->Get()->TryGetNumberField(TEXT("z"), NewPoint.Z);
+
+					NewRow->ObjectPoints.Points.Add(NewPoint);
+				}
+
+				const TSharedPtr<FJsonObject>* Point2DJsonObject;
+				if (PointDataJsonObject->Get()->TryGetObjectField(TEXT("point2d"), Point2DJsonObject))
+				{
+					FVector2D NewPoint;
+					Point2DJsonObject->Get()->TryGetNumberField(TEXT("x"), NewPoint.X);
+					Point2DJsonObject->Get()->TryGetNumberField(TEXT("y"), NewPoint.Y);
+
+					NewRow->ImagePoints.Points.Add(NewPoint);
+				}
+			}
+		}
+
+		const TSharedPtr<FJsonObject>* CameraDataJsonObject;
+		if (CalibrationRowObject->TryGetObjectField(TEXT("cameraData"), CameraDataJsonObject))
+		{
+			const TSharedPtr<FJsonObject>* CameraPoseJsonObject;
+			if (CameraDataJsonObject->Get()->TryGetObjectField(TEXT("pose"), CameraPoseJsonObject))
+			{
+				const TSharedPtr<FJsonObject>* RotationJsonObject;
+				if (CameraPoseJsonObject->Get()->TryGetObjectField(TEXT("rotation"), RotationJsonObject))
+				{
+					FQuat Rotation;
+					RotationJsonObject->Get()->TryGetNumberField(TEXT("x"), Rotation.X);
+					RotationJsonObject->Get()->TryGetNumberField(TEXT("y"), Rotation.Y);
+					RotationJsonObject->Get()->TryGetNumberField(TEXT("z"), Rotation.Z);
+					RotationJsonObject->Get()->TryGetNumberField(TEXT("w"), Rotation.W);
+					NewRow->CameraPose.SetRotation(Rotation);
+				}
+
+				const TSharedPtr<FJsonObject>* TranslationJsonObject;
+				if (CameraPoseJsonObject->Get()->TryGetObjectField(TEXT("translation"), TranslationJsonObject))
+				{
+					FVector Translation;
+					TranslationJsonObject->Get()->TryGetNumberField(TEXT("x"), Translation.X);
+					TranslationJsonObject->Get()->TryGetNumberField(TEXT("y"), Translation.Y);
+					TranslationJsonObject->Get()->TryGetNumberField(TEXT("z"), Translation.Z);
+					NewRow->CameraPose.SetTranslation(Translation);
+				}
+			}
 		}
 
 		Dataset.CalibrationRows.Add(NewRow);
 	}
-	else
+	else if (DatasetVersion == EDatasetVersion::CombinedAlgoClasses)
 	{
-		UE_LOG(LogCameraCalibrationEditor, Warning, TEXT("Failed to import calibration row because at least one field could not be deserialized from the json file."));
+		// We enforce strict mode to ensure that every field in the UStruct of row data is present in the imported json.
+		// If any fields are missing, it is likely the row will be invalid, which will lead to errors in the calibration.
+		constexpr int64 CheckFlags = 0;
+		constexpr int64 SkipFlags = 0;
+		constexpr bool bStrictMode = true;
+		if (FJsonObjectConverter::JsonObjectToUStruct<FCalibrationRow>(CalibrationRowObject, NewRow.Get(), CheckFlags, SkipFlags, bStrictMode))
+		{
+			CaptureSettings.CalibrationPattern = NewRow->Pattern;
+			Dataset.CalibrationRows.Add(NewRow);
+		}
+		else
+		{
+			UE_LOG(LogCameraCalibrationEditor, Warning, TEXT("Failed to import calibration row because at least one field could not be deserialized from the json file."));
+		}
 	}
 
 	return NewRow->Index;
diff --git a/Engine/Plugins/VirtualProduction/CameraCalibration/Source/CameraCalibrationEditor/Private/AssetEditor/LensDistortionTool.h b/Engine/Plugins/VirtualProduction/CameraCalibration/Source/CameraCalibrationEditor/Private/AssetEditor/LensDistortionTool.h
--- a/Engine/Plugins/VirtualProduction/CameraCalibration/Source/CameraCalibrationEditor/Private/AssetEditor/LensDistortionTool.h
+++ b/Engine/Plugins/VirtualProduction/CameraCalibration/Source/CameraCalibrationEditor/Private/AssetEditor/LensDistortionTool.h
@@ -38,6 +38,17 @@ enum class ECalibrationType : uint8
 	Everything
 };
 
+/** Version info to support backwards compatibility of dataset importing */
+enum class EDatasetVersion : uint8
+{
+	Invalid = 0,
+	SeparateAlgoClasses = 1,
+	CombinedAlgoClasses = 2,
+
+	// Add new versions above
+	CurrentVersion = CombinedAlgoClasses
+};
+
 /** The data associated with a single captured calibration pattern/point */
 USTRUCT()
 struct FCalibrationRow
@@ -251,7 +262,7 @@ private:
 	 * Import a JsonObject of calibration data that represents a single calibration row.
 	 * Returns the row index of the imported row.
 	 */
-	int32 ImportCalibrationRow(const TSharedRef<FJsonObject>& CalibrationRowObject, const FImage& RowImage);
+	int32 ImportCalibrationRow(const TSharedRef<FJsonObject>& CalibrationRowObject, const FImage& RowImage, EDatasetVersion DatasetVersion);
 
 	/** Export global session data to a .json file */
 	void ExportSessionData();
diff --git a/Engine/Plugins/VirtualProduction/CameraCalibration/Source/CameraCalibrationEditor/Private/AssetEditor/SLensDistortionToolPanel.cpp b/Engine/Plugins/VirtualProduction/CameraCalibration/Source/CameraCalibrationEditor/Private/AssetEditor/SLensDistortionToolPanel.cpp
--- a/Engine/Plugins/VirtualProduction/CameraCalibration/Source/CameraCalibrationEditor/Private/AssetEditor/SLensDistortionToolPanel.cpp
+++ b/Engine/Plugins/VirtualProduction/CameraCalibration/Source/CameraCalibrationEditor/Private/AssetEditor/SLensDistortionToolPanel.cpp
@@ -13,7 +13,6 @@
 #include "UI/SImageTexture.h"
 #include "Widgets/Input/SButton.h"
 #include "Widgets/SBoxPanel.h"
-#include "Widgets/SWindow.h"
 
 #define LOCTEXT_NAMESPACE "SLensDistortionToolPanel"
 
@@ -296,9 +295,7 @@ void SLensDistortionToolPanel::OnCaptureSettingsChanged(const FPropertyChangedEv
 	}
 	else if (PropertyName == GET_MEMBER_NAME_CHECKED(FLensCaptureSettings, CalibrationPattern))
 	{
-		DatasetListHeader->SetShowGeneratedColumn("Image", (Tool->CaptureSettings.CalibrationPattern != ECalibrationPattern::Points));
-		DatasetListHeader->SetShowGeneratedColumn("ImagePoint", (Tool->CaptureSettings.CalibrationPattern == ECalibrationPattern::Points));
-		DatasetListHeader->SetShowGeneratedColumn("ObjectPoint", (Tool->CaptureSettings.CalibrationPattern == ECalibrationPattern::Points));
+		RefreshListView();
 	}
 }
 
@@ -414,6 +411,10 @@ void SLensDistortionToolPanel::RefreshListView()
 {
 	if (DatasetListView)
 	{
+		DatasetListHeader->SetShowGeneratedColumn("Image", (Tool->CaptureSettings.CalibrationPattern != ECalibrationPattern::Points));
+		DatasetListHeader->SetShowGeneratedColumn("ImagePoint", (Tool->CaptureSettings.CalibrationPattern == ECalibrationPattern::Points));
+		DatasetListHeader->SetShowGeneratedColumn("ObjectPoint", (Tool->CaptureSettings.CalibrationPattern == ECalibrationPattern::Points));
+
 		DatasetListView->RequestListRefresh();
 	}
 }
