CameraCalibration: Fix crash caused by overlay texture getting garbage collected. Also small fix to the debug checkerboard corners being drawn to the wrong coordinates.

diff --git a/Engine/Plugins/VirtualProduction/CameraCalibration/Source/CameraCalibrationEditor/Private/AssetEditor/LensDistortionTool.cpp b/Engine/Plugins/VirtualProduction/CameraCalibration/Source/CameraCalibrationEditor/Private/AssetEditor/LensDistortionTool.cpp
--- a/Engine/Plugins/VirtualProduction/CameraCalibration/Source/CameraCalibrationEditor/Private/AssetEditor/LensDistortionTool.cpp
+++ b/Engine/Plugins/VirtualProduction/CameraCalibration/Source/CameraCalibrationEditor/Private/AssetEditor/LensDistortionTool.cpp
@@ -284,7 +284,7 @@ bool ULensDistortionTool::DetectCheckerboardPattern(TArray<FColor>& Pixels, FInt
 	const FIntPoint OverlayTextureSize = FIntPoint(OverlayTexture->GetSizeX(), OverlayTexture->GetSizeY());
 	RescalePoints(CameraFeedAdjustedCorners, OverlayTextureSize, Size);
 
-	FOpenCVHelper::DrawCheckerboardCorners(DetectedCorners, OutRow->CheckerboardDimensions, OverlayTexture);
+	FOpenCVHelper::DrawCheckerboardCorners(CameraFeedAdjustedCorners, OutRow->CheckerboardDimensions, OverlayTexture);
 
 	if (TSharedPtr<FCameraCalibrationStepsController> StepsController = WeakStepsController.Pin())
 	{
diff --git a/Engine/Plugins/VirtualProduction/CameraCalibration/Source/CameraCalibrationEditor/Private/AssetEditor/LensDistortionTool.h b/Engine/Plugins/VirtualProduction/CameraCalibration/Source/CameraCalibrationEditor/Private/AssetEditor/LensDistortionTool.h
--- a/Engine/Plugins/VirtualProduction/CameraCalibration/Source/CameraCalibrationEditor/Private/AssetEditor/LensDistortionTool.h
+++ b/Engine/Plugins/VirtualProduction/CameraCalibration/Source/CameraCalibrationEditor/Private/AssetEditor/LensDistortionTool.h
@@ -281,6 +281,7 @@ private:
 	FCalibrationDataset Dataset;
 
 	/** The solver object that runs the lens calibration on another thread. The ptr is stored to access the solver to get status updates and/or cancel the calibration from the game thread */
+	UPROPERTY(Transient)
 	TObjectPtr<ULensDistortionSolver> Solver;
 
 	/** An asynchronous task handle. When valid, the tool will poll its state to determine when the task has completed, and then extract the calibration result from this task handle. */
@@ -308,7 +309,10 @@ private:
 	bool bIsActive = false;
 
 	/** Material and Texture that draw the coverage overlay on top of the simulcam viewport during capture */
+	UPROPERTY(Transient)
 	TObjectPtr<UMaterialInstanceDynamic> OverlayMID;
+
+	UPROPERTY(Transient)
 	TObjectPtr<UTexture2D> OverlayTexture;
 
 	/** Weak ptr to the steps controller that created this step */
