CameraCalibration: Fix bug when drawing checkerboard corners after media has been resized.
#rnx

diff --git a/Engine/Plugins/VirtualProduction/CameraCalibration/Source/CameraCalibrationEditor/Private/AssetEditor/LensDistortionTool.cpp b/Engine/Plugins/VirtualProduction/CameraCalibration/Source/CameraCalibrationEditor/Private/AssetEditor/LensDistortionTool.cpp
--- a/Engine/Plugins/VirtualProduction/CameraCalibration/Source/CameraCalibrationEditor/Private/AssetEditor/LensDistortionTool.cpp
+++ b/Engine/Plugins/VirtualProduction/CameraCalibration/Source/CameraCalibrationEditor/Private/AssetEditor/LensDistortionTool.cpp
@@ -95,6 +95,17 @@ TSharedRef<SWidget> ULensDistortionTool::BuildUI()
 
 void ULensDistortionTool::Tick(float DeltaTime)
 {
+	// If the resolution of the simulcam comp has changed, update the coverage texture to be the correct size
+	if (TSharedPtr<FCameraCalibrationStepsController> StepsController = WeakStepsController.Pin())
+	{
+		const FIntPoint Size = StepsController->GetCompRenderResolution();
+		if (OverlayTexture && ((OverlayTexture->GetSizeX() != Size.X) || (OverlayTexture->GetSizeY() != Size.Y)))
+		{
+			OverlayTexture = UTexture2D::CreateTransient(Size.X, Size.Y, EPixelFormat::PF_B8G8R8A8);
+			RefreshCoverage();
+		}
+	}
+
 	// A valid task handle implies that there is an asynchronous calibration happening on another thread.
 	if (CalibrationTask.IsValid())
 	{
