CameraCalibration: Replace unclear "everything" option from distortion solver settings with a checkbox to enable/disable solving for nodal offset

diff --git a/Engine/Plugins/VirtualProduction/CameraCalibration/Source/CameraCalibrationEditor/Private/AssetEditor/LensDistortionTool.cpp b/Engine/Plugins/VirtualProduction/CameraCalibration/Source/CameraCalibrationEditor/Private/AssetEditor/LensDistortionTool.cpp
--- a/Engine/Plugins/VirtualProduction/CameraCalibration/Source/CameraCalibrationEditor/Private/AssetEditor/LensDistortionTool.cpp
+++ b/Engine/Plugins/VirtualProduction/CameraCalibration/Source/CameraCalibrationEditor/Private/AssetEditor/LensDistortionTool.cpp
@@ -638,7 +638,7 @@ void ULensDistortionTool::SaveCalibrationResult()
 	}
 
 	// If the calibration result contains a meaningful nodal offset result, add it to the Lens File
-	if (SolverSettings.CalibrationType == ECalibrationType::NodalOffset || SolverSettings.CalibrationType == ECalibrationType::Everything)
+	if (SolverSettings.bSolveNodalOffset && CaptureSettings.bIsCalibratorTracked && CaptureSettings.bIsCameraTracked)
 	{
 		if (!CalibrationResult.NodalOffset.LocationOffset.Equals(FVector::ZeroVector) || !CalibrationResult.NodalOffset.RotationOffset.Equals(FQuat::Identity))
 		{
diff --git a/Engine/Plugins/VirtualProduction/CameraCalibration/Source/CameraCalibrationEditor/Private/AssetEditor/LensDistortionTool.h b/Engine/Plugins/VirtualProduction/CameraCalibration/Source/CameraCalibrationEditor/Private/AssetEditor/LensDistortionTool.h
--- a/Engine/Plugins/VirtualProduction/CameraCalibration/Source/CameraCalibrationEditor/Private/AssetEditor/LensDistortionTool.h
+++ b/Engine/Plugins/VirtualProduction/CameraCalibration/Source/CameraCalibrationEditor/Private/AssetEditor/LensDistortionTool.h
@@ -29,15 +29,6 @@ enum class ECalibrationPattern : uint8
 	Points
 };
 
-/** List of supported calibration workflows */
-UENUM()
-enum class ECalibrationType : uint8
-{
-	DistortionAndIntrinsics,
-	NodalOffset,
-	Everything
-};
-
 /** Version info to support backwards compatibility of dataset importing */
 enum class EDatasetVersion : uint8
 {
@@ -137,9 +128,12 @@ struct FLensSolverSettings
 	UPROPERTY(EditAnywhere, NoClear, Category = "Solver Settings", meta = (HideViewOptions, ShowDisplayNames))
 	TSubclassOf<ULensDistortionSolver> SolverClass = ULensDistortionSolverOpenCV::StaticClass();
 
-	/** Select the data that the solver should compute and write to the Lens File */
+	/** 
+	 * If true, the solver will calibrate for the nodal offset in addition to lens distortion. 
+	 * This property will be read-only if either IsCalibratorTracked or IsCameraTracked is unchecked because both are required to compute the nodal offset.
+	 */
 	UPROPERTY(EditAnywhere, Category = "Solver Settings")
-	ECalibrationType CalibrationType = ECalibrationType::DistortionAndIntrinsics;
+	bool bSolveNodalOffset = false;
 
 	/** An estimate for the focal length of the lens */
 	UPROPERTY(EditAnywhere, Category = "Solver Settings")
diff --git a/Engine/Plugins/VirtualProduction/CameraCalibration/Source/CameraCalibrationEditor/Private/AssetEditor/SLensDistortionToolPanel.cpp b/Engine/Plugins/VirtualProduction/CameraCalibration/Source/CameraCalibrationEditor/Private/AssetEditor/SLensDistortionToolPanel.cpp
--- a/Engine/Plugins/VirtualProduction/CameraCalibration/Source/CameraCalibrationEditor/Private/AssetEditor/SLensDistortionToolPanel.cpp
+++ b/Engine/Plugins/VirtualProduction/CameraCalibration/Source/CameraCalibrationEditor/Private/AssetEditor/SLensDistortionToolPanel.cpp
@@ -156,6 +156,15 @@ FAssetData FCaptureSettingsCustomization::GetCalibratorAssetData() const
 	return Tmp;
 }
 
+bool SLensDistortionToolPanel::IsSolverSettingPropertyReadOnly(const FPropertyAndParent& PropertyAndParent) const
+{
+	if (PropertyAndParent.Property.GetFName() == GET_MEMBER_NAME_CHECKED(FLensSolverSettings, bSolveNodalOffset) && Tool.IsValid())
+	{
+		return !(Tool->CaptureSettings.bIsCalibratorTracked && Tool->CaptureSettings.bIsCameraTracked);
+	}
+	return false;
+}
+
 void SLensDistortionToolPanel::Construct(const FArguments& InArgs, ULensDistortionTool* InTool, TWeakPtr<FCameraCalibrationStepsController> InStepsController)
 {
 	Tool = InTool;
@@ -171,6 +180,7 @@ void SLensDistortionToolPanel::Construct(const FArguments& InArgs, ULensDistorti
 
 	TSharedRef<FStructOnScope> SolverSettingsStruct = MakeShared<FStructOnScope>(FLensSolverSettings::StaticStruct(), reinterpret_cast<uint8*>(&Tool->SolverSettings));
 	TSharedPtr<IStructureDetailsView> SolverSettingsDetailsView = PropertyEditor.CreateStructureDetailView(DefaultDetailsViewArgs, DefaultStructuDetailsViewArgs, SolverSettingsStruct);
+	SolverSettingsDetailsView->GetDetailsView()->SetIsPropertyReadOnlyDelegate(FIsPropertyReadOnly::CreateSP(this, &SLensDistortionToolPanel::IsSolverSettingPropertyReadOnly));
 
 	TSharedRef<FStructOnScope> CaptureSettingsStruct = MakeShared<FStructOnScope>(FLensCaptureSettings::StaticStruct(), reinterpret_cast<uint8*>(&Tool->CaptureSettings));
 
diff --git a/Engine/Plugins/VirtualProduction/CameraCalibration/Source/CameraCalibrationEditor/Private/AssetEditor/SLensDistortionToolPanel.h b/Engine/Plugins/VirtualProduction/CameraCalibration/Source/CameraCalibrationEditor/Private/AssetEditor/SLensDistortionToolPanel.h
--- a/Engine/Plugins/VirtualProduction/CameraCalibration/Source/CameraCalibrationEditor/Private/AssetEditor/SLensDistortionToolPanel.h
+++ b/Engine/Plugins/VirtualProduction/CameraCalibration/Source/CameraCalibrationEditor/Private/AssetEditor/SLensDistortionToolPanel.h
@@ -77,6 +77,7 @@ public:
 private:
 	/** Callback to handle changes to the customized capture settings */
 	void OnCaptureSettingsChanged(const FPropertyChangedEvent& PropertyChangedEvent);
+	bool IsSolverSettingPropertyReadOnly(const FPropertyAndParent& PropertyAndParent) const;
 
 	TSharedRef<SWidget> BuildDatasetListView();
 	TSharedRef<ITableRow> OnGenerateDatasetRow(TSharedPtr<FCalibrationRow> InItem, const TSharedRef<STableViewBase>& OwnerTable);
