CameraCalibration: Prevent checkerboard detection from blocking the game thead for long periods of time in rare cases when opencv does not return quickly.

diff --git a/Engine/Plugins/VirtualProduction/CameraCalibration/Source/CameraCalibrationEditor/Private/AssetEditor/LensDistortionTool.cpp b/Engine/Plugins/VirtualProduction/CameraCalibration/Source/CameraCalibrationEditor/Private/AssetEditor/LensDistortionTool.cpp
--- a/Engine/Plugins/VirtualProduction/CameraCalibration/Source/CameraCalibrationEditor/Private/AssetEditor/LensDistortionTool.cpp
+++ b/Engine/Plugins/VirtualProduction/CameraCalibration/Source/CameraCalibrationEditor/Private/AssetEditor/LensDistortionTool.cpp
@@ -253,10 +253,38 @@ bool ULensDistortionTool::DetectCheckerboardPattern(TArray<FColor>& Pixels, FInt
 		return false;
 	}
 
-	OutRow->CheckerboardDimensions = FIntPoint(Checkerboard->NumCornerCols, Checkerboard->NumCornerRows);
+	const FIntPoint CheckerboardDimensions = FIntPoint(Checkerboard->NumCornerCols, Checkerboard->NumCornerRows);
+	OutRow->CheckerboardDimensions = CheckerboardDimensions;
+
+	// Launch an async task to perform the opencv checkerboard detection to prevent the game thread from being blocked in the rare cases when detection takes a very long time
+	UE::Tasks::TTask<TArray<FVector2f>> DetectionTask = UE::Tasks::Launch(UE_SOURCE_LOCATION, [Pixels, Size, CheckerboardDimensions]() mutable
+	{
+		TArray<FVector2f> Corners;
+		FOpenCVHelper::IdentifyCheckerboard(Pixels, Size, CheckerboardDimensions, Corners);
+		return Corners;
+	});
 
 	TArray<FVector2f> DetectedCorners;
-	const bool bCornersFound = FOpenCVHelper::IdentifyCheckerboard(Pixels, Size, OutRow->CheckerboardDimensions, DetectedCorners);
+	bool bCornersFound = false;
+
+	const FDateTime StartTime = FDateTime::Now();
+
+	const float Timeout = GetDefault<UCameraCalibrationSettings>()->GetCheckerboardDetectionTimeout();
+
+	// If the detection has not completed before a set timeout, abandon this task. The user will be informed of the detection failure.
+	while ((FDateTime::Now() - StartTime).GetSeconds() < Timeout)
+	{
+		if (DetectionTask.IsValid() && DetectionTask.IsCompleted())
+		{
+			// Extract the return value from the task
+			DetectedCorners = DetectionTask.GetResult();
+			bCornersFound = DetectedCorners.Num() > 0;
+			break;
+		}
+	}
+
+	// Release the task resource handle
+	DetectionTask = {};
 
 	if (!bCornersFound || DetectedCorners.IsEmpty())
 	{
diff --git a/Engine/Plugins/VirtualProduction/CameraCalibrationCore/Source/CameraCalibrationCore/Public/CameraCalibrationSettings.h b/Engine/Plugins/VirtualProduction/CameraCalibrationCore/Source/CameraCalibrationCore/Public/CameraCalibrationSettings.h
--- a/Engine/Plugins/VirtualProduction/CameraCalibrationCore/Source/CameraCalibrationCore/Public/CameraCalibrationSettings.h
+++ b/Engine/Plugins/VirtualProduction/CameraCalibrationCore/Source/CameraCalibrationCore/Public/CameraCalibrationSettings.h
@@ -64,6 +64,9 @@ public:
 	/** Get the tolerance to use when adding or accessing data in a calibrated LensFile */
 	float GetCalibrationInputTolerance() const { return CalibrationInputTolerance; }
 
+	/** Get the length of time (in seconds) the calibration tool will wait when attempting to detect a checkerboard pattern in the media image */
+	float GetCheckerboardDetectionTimeout() const { return CheckerboardDetectionTimeout; }
+
 	/** Get the default MaterialInterface used by the input Model Handler class to write the undistortion displacement map */
 	UMaterialInterface* GetDefaultUndistortionDisplacementMaterial(const TSubclassOf<ULensDistortionModelHandlerBase>& InModelHandler) const;
 
@@ -101,6 +104,10 @@ private:
 	UPROPERTY(config, EditAnywhere, Category = "Settings")
 	float CalibrationInputTolerance = 0.001f;
 
+	/** Length of time (in seconds) the calibration tool will wait when attempting to detect a checkerboard pattern in the media image */
+	UPROPERTY(config, EditAnywhere, Category = "Settings")
+	float CheckerboardDetectionTimeout = 2.0f;
+
 	/** Map of Lens Distortion Model Handler classes to the default displacement map material used by that class */
 	UPROPERTY(config)
 	TMap<TSubclassOf<ULensDistortionModelHandlerBase>, TSoftObjectPtr<UMaterialInterface>> DefaultUndistortionDisplacementMaterials;
