CameraCalibration: Fix distortion calibration bug where camera poses are grouped together, even when the calibrator is not tracked. Fix another calibration bug where the solver does not converge when tracked calibrator poses are used in the absence of tracked camera poses.

diff --git a/Engine/Plugins/VirtualProduction/CameraCalibration/Source/CameraCalibrationEditor/Private/AssetEditor/LensDistortionTool.cpp b/Engine/Plugins/VirtualProduction/CameraCalibration/Source/CameraCalibrationEditor/Private/AssetEditor/LensDistortionTool.cpp
--- a/Engine/Plugins/VirtualProduction/CameraCalibration/Source/CameraCalibrationEditor/Private/AssetEditor/LensDistortionTool.cpp
+++ b/Engine/Plugins/VirtualProduction/CameraCalibration/Source/CameraCalibrationEditor/Private/AssetEditor/LensDistortionTool.cpp
@@ -543,12 +543,21 @@ void ULensDistortionTool::CalibrateLens()
 	{
 		EnumAddFlags(SolverFlags, ECalibrationFlags::UseExtrinsicGuess);
 	}
+	else if (CaptureSettings.CalibrationPattern == ECalibrationPattern::Checkerboard)
+	{
+		GenerateDummyCheckerboardPoints(Samples3d, Dataset.CalibrationRows.Num(), Dataset.CalibrationRows[0]->CheckerboardDimensions);
+	}
 
 	if (CaptureSettings.bIsCalibratorTracked && CaptureSettings.CalibrationPattern != ECalibrationPattern::Points)
 	{
 		EnumAddFlags(SolverFlags, ECalibrationFlags::SolveTargetOffset);
 	}
 
+	if (CaptureSettings.bIsCameraTracked && CaptureSettings.bIsCalibratorTracked)
+	{
+		EnumAddFlags(SolverFlags, ECalibrationFlags::GroupCameraPoses);
+	}
+
 	if (SolverSettings.bFixFocalLength)
 	{
 		EnumAddFlags(SolverFlags, ECalibrationFlags::FixFocalLength);
@@ -829,6 +838,27 @@ void ULensDistortionTool::RescalePoints(TArray<FVector2D>& Points, FIntPoint Deb
 	}
 }
 
+void ULensDistortionTool::GenerateDummyCheckerboardPoints(TArray<FObjectPoints>& Samples3d, int32 NumImages, FIntPoint CheckerboardDimensions)
+{
+	// If the camera is not tracked, the distortion solver must initialize the camera pose for each image using linear algebra techniques.
+	// However, it struggles to do so when "real" tracking data is used for the calibrator. So in this case, we replace the tracked calibrator data 
+	// with a set of dummy points for the 3D checkerboard corners. The board is assumed to lie in the YZ plane with the TopLeft corner at (0, 0, 0) in world space.
+	Samples3d.Empty();
+	for (int32 ImageIndex = 0; ImageIndex < NumImages; ++ImageIndex)
+	{
+		FObjectPoints Points3d;
+		for (int32 RowIdx = 0; RowIdx < CheckerboardDimensions.Y; ++RowIdx)
+		{
+			for (int32 ColIdx = 0; ColIdx < CheckerboardDimensions.X; ++ColIdx)
+			{
+				Points3d.Points.Add(FVector(0, ColIdx, -RowIdx));
+			}
+		}
+
+		Samples3d.Add(Points3d);
+	}
+}
+
 bool ULensDistortionTool::DependsOnStep(UCameraCalibrationStep* Step) const
 {
 	return Cast<ULensInfoStep>(Step) != nullptr;
diff --git a/Engine/Plugins/VirtualProduction/CameraCalibration/Source/CameraCalibrationEditor/Private/AssetEditor/LensDistortionTool.h b/Engine/Plugins/VirtualProduction/CameraCalibration/Source/CameraCalibrationEditor/Private/AssetEditor/LensDistortionTool.h
--- a/Engine/Plugins/VirtualProduction/CameraCalibration/Source/CameraCalibrationEditor/Private/AssetEditor/LensDistortionTool.h
+++ b/Engine/Plugins/VirtualProduction/CameraCalibration/Source/CameraCalibrationEditor/Private/AssetEditor/LensDistortionTool.h
@@ -246,6 +246,9 @@ private:
 	/** Adjust the 2D location of every input pixel to account for differences between the size of the captured image and the overlay texture */
 	void RescalePoints(TArray<FVector2D>& Points, FIntPoint DebugTextureSize, FIntPoint CameraFeedSize);
 
+	/** Replace tracked checkerboard calibration points with a set of dummy 3D points */
+	void GenerateDummyCheckerboardPoints(TArray<FObjectPoints>& Samples3d, int32 NumImages, FIntPoint CheckerboardDimensions);
+
 	/** Get the directory for the current session */
 	FString GetSessionSaveDir() const;
 
diff --git a/Engine/Plugins/VirtualProduction/CameraCalibration/Source/CameraCalibrationEditor/Private/Calibrators/CameraCalibrationSolver.cpp b/Engine/Plugins/VirtualProduction/CameraCalibration/Source/CameraCalibrationEditor/Private/Calibrators/CameraCalibrationSolver.cpp
--- a/Engine/Plugins/VirtualProduction/CameraCalibration/Source/CameraCalibrationEditor/Private/Calibrators/CameraCalibrationSolver.cpp
+++ b/Engine/Plugins/VirtualProduction/CameraCalibration/Source/CameraCalibrationEditor/Private/Calibrators/CameraCalibrationSolver.cpp
@@ -86,10 +86,22 @@ FDistortionCalibrationResult ULensDistortionSolverOpenCV::Solve_Implementation(
 	// Reorganize the 3D and 2D points from the input arrays or arrays to be laid out linearly in memory in two cv::Mat objects
 	GatherPoints(ObjectPointArray, ImagePointArray, ObjectPointsMat, ImagePointsMat);
 
-	// Find the set of unique camera poses to reduce the number of poses we need to solve for
 	TArray<int32> CameraPoseIndices;
 	TArray<FTransform> UniqueCameraPoses;
-	FindUniqueCameraPoses(CameraPoses, UniqueCameraPoses, CameraPoseIndices);
+	if ((EnumHasAnyFlags(SolverFlags, ECalibrationFlags::GroupCameraPoses)))
+	{
+		// Find the set of unique camera poses to reduce the number of poses we need to solve for
+		FindUniqueCameraPoses(CameraPoses, UniqueCameraPoses, CameraPoseIndices);
+	}
+	else
+	{
+		// Do not group camera poses, and consider each one to be unique
+		for (int32 PoseIndex = 0; PoseIndex < CameraPoses.Num(); ++PoseIndex)
+		{
+			UniqueCameraPoses.Add(CameraPoses[PoseIndex]);
+			CameraPoseIndices.Add(PoseIndex);
+		}
+	}
 
 	double RMSE = 0.0;
 
diff --git a/Engine/Plugins/VirtualProduction/CameraCalibration/Source/CameraCalibrationEditor/Private/Calibrators/CameraCalibrationSolver.h b/Engine/Plugins/VirtualProduction/CameraCalibration/Source/CameraCalibrationEditor/Private/Calibrators/CameraCalibrationSolver.h
--- a/Engine/Plugins/VirtualProduction/CameraCalibration/Source/CameraCalibrationEditor/Private/Calibrators/CameraCalibrationSolver.h
+++ b/Engine/Plugins/VirtualProduction/CameraCalibration/Source/CameraCalibrationEditor/Private/Calibrators/CameraCalibrationSolver.h
@@ -58,7 +58,8 @@ enum class ECalibrationFlags : uint32
 	FixExtrinsics = 1 << 4,      /** The solver will not optimize the camera extrinsics */
 	FixDistortion = 1 << 5,      /** The solver will fix all distortion values at 0 */
 	FixAspectRatio = 1 << 6,     /** The solver will respect the input aspect ratio when solving for Fx and Fy */
-	SolveTargetOffset = 1 << 7   /** The solver will solve for an offset for each image's input 3D object points */
+	SolveTargetOffset = 1 << 7,  /** The solver will solve for an offset for each image's input 3D object points */
+	GroupCameraPoses = 1 << 8    /** The solver will group sets of images together that share the same camera pose */
 };
 
 /** 
