CameraCalibration: Add support for marquee select events to the simulcam viewport. This CL lays the groundwork for another change to allow users to perform marquee selection to do pattern detection on subregions of their media image.

diff --git a/Engine/Plugins/VirtualProduction/CameraCalibration/Source/CameraCalibrationEditor/Private/AssetEditor/CameraCalibrationStepsController.cpp b/Engine/Plugins/VirtualProduction/CameraCalibration/Source/CameraCalibrationEditor/Private/AssetEditor/CameraCalibrationStepsController.cpp
--- a/Engine/Plugins/VirtualProduction/CameraCalibration/Source/CameraCalibrationEditor/Private/AssetEditor/CameraCalibrationStepsController.cpp
+++ b/Engine/Plugins/VirtualProduction/CameraCalibration/Source/CameraCalibrationEditor/Private/AssetEditor/CameraCalibrationStepsController.cpp
@@ -1218,6 +1218,18 @@ bool FCameraCalibrationStepsController::OnSimulcamViewportInputKey(const FKey& I
 	return bStepHandled;
 }
 
+void FCameraCalibrationStepsController::OnSimulcamViewportMarqueeSelect(FVector2D StartPosition, FVector2D EndPosition)
+{
+	for (TStrongObjectPtr<UCameraCalibrationStep>& Step : CalibrationSteps)
+	{
+		if (Step.IsValid() && Step->IsActive())
+		{
+			Step->OnViewportMarqueeSelect(StartPosition, EndPosition);
+			break;
+		}
+	}
+}
+
 FReply FCameraCalibrationStepsController::OnRewindButtonClicked()
 {
 	// Rewind to the beginning of the media
diff --git a/Engine/Plugins/VirtualProduction/CameraCalibration/Source/CameraCalibrationEditor/Private/AssetEditor/CameraCalibrationStepsController.h b/Engine/Plugins/VirtualProduction/CameraCalibration/Source/CameraCalibrationEditor/Private/AssetEditor/CameraCalibrationStepsController.h
--- a/Engine/Plugins/VirtualProduction/CameraCalibration/Source/CameraCalibrationEditor/Private/AssetEditor/CameraCalibrationStepsController.h
+++ b/Engine/Plugins/VirtualProduction/CameraCalibration/Source/CameraCalibrationEditor/Private/AssetEditor/CameraCalibrationStepsController.h
@@ -183,6 +183,9 @@ public:
 	/** Called by the UI when the Simulcam Viewport receives keyboard input */
 	bool OnSimulcamViewportInputKey(const FKey& InKey, const EInputEvent& InEvent);
 
+	/** Called by the UI when the Simulcam Viewport receives a marquee select event */
+	void OnSimulcamViewportMarqueeSelect(FVector2D StartPosition, FVector2D EndPosition);
+
 	/** Called by the UI when the rewind button is clicked */
 	FReply OnRewindButtonClicked();
 
diff --git a/Engine/Plugins/VirtualProduction/CameraCalibration/Source/CameraCalibrationEditor/Private/AssetEditor/SCameraCalibrationSteps.cpp b/Engine/Plugins/VirtualProduction/CameraCalibration/Source/CameraCalibrationEditor/Private/AssetEditor/SCameraCalibrationSteps.cpp
--- a/Engine/Plugins/VirtualProduction/CameraCalibration/Source/CameraCalibrationEditor/Private/AssetEditor/SCameraCalibrationSteps.cpp
+++ b/Engine/Plugins/VirtualProduction/CameraCalibration/Source/CameraCalibrationEditor/Private/AssetEditor/SCameraCalibrationSteps.cpp
@@ -242,6 +242,7 @@ void SCameraCalibrationSteps::Construct(const FArguments& InArgs, TWeakPtr<FCame
 				SNew(SSimulcamViewport, CalibrationStepsController.Pin()->GetRenderTarget())
 				.OnSimulcamViewportClicked_Raw(CalibrationStepsController.Pin().Get(), &FCameraCalibrationStepsController::OnSimulcamViewportClicked)
 				.OnSimulcamViewportInputKey_Raw(CalibrationStepsController.Pin().Get(), &FCameraCalibrationStepsController::OnSimulcamViewportInputKey)
+				.OnSimulcamViewportMarqueeSelect_Raw(CalibrationStepsController.Pin().Get(), &FCameraCalibrationStepsController::OnSimulcamViewportMarqueeSelect)
 			]
 
 			+ SVerticalBox::Slot() // Media playback buttons
diff --git a/Engine/Plugins/VirtualProduction/CameraCalibration/Source/CameraCalibrationEditor/Private/AssetEditor/SSimulcamViewport.cpp b/Engine/Plugins/VirtualProduction/CameraCalibration/Source/CameraCalibrationEditor/Private/AssetEditor/SSimulcamViewport.cpp
--- a/Engine/Plugins/VirtualProduction/CameraCalibration/Source/CameraCalibrationEditor/Private/AssetEditor/SSimulcamViewport.cpp
+++ b/Engine/Plugins/VirtualProduction/CameraCalibration/Source/CameraCalibrationEditor/Private/AssetEditor/SSimulcamViewport.cpp
@@ -9,6 +9,7 @@ void SSimulcamViewport::Construct(const FArguments& InArgs, UTexture* InTexture)
 {
 	OnSimulcamViewportClicked = InArgs._OnSimulcamViewportClicked;
 	OnSimulcamViewportInputKey = InArgs._OnSimulcamViewportInputKey;
+	OnSimulcamViewportMarqueeSelect = InArgs._OnSimulcamViewportMarqueeSelect;
 
 	Texture = TStrongObjectPtr<UTexture>(InTexture);
 
diff --git a/Engine/Plugins/VirtualProduction/CameraCalibration/Source/CameraCalibrationEditor/Private/AssetEditor/SSimulcamViewport.h b/Engine/Plugins/VirtualProduction/CameraCalibration/Source/CameraCalibrationEditor/Private/AssetEditor/SSimulcamViewport.h
--- a/Engine/Plugins/VirtualProduction/CameraCalibration/Source/CameraCalibrationEditor/Private/AssetEditor/SSimulcamViewport.h
+++ b/Engine/Plugins/VirtualProduction/CameraCalibration/Source/CameraCalibrationEditor/Private/AssetEditor/SSimulcamViewport.h
@@ -23,6 +23,9 @@ DECLARE_DELEGATE_TwoParams(FSimulcamViewportClickedEventHandler, const FGeometry
 /** Delegate to be executed when the viewport area receives a key event */
 DECLARE_DELEGATE_RetVal_TwoParams(bool, FSimulcamViewportInputKeyEventHandler, const FKey& Key, const EInputEvent& Event);
 
+/** Delegate to be executed when the user performs a marquee selection on the viewport area */
+DECLARE_DELEGATE_TwoParams(FSimulcamViewportMarqueeSelectEventHandler, FVector2D StartPosition, FVector2D EndPosition);
+
 /**
  * UI to display the provided UTexture and respond to user input.
  */
@@ -34,6 +37,7 @@ public:
 	{}
 	SLATE_EVENT(FSimulcamViewportClickedEventHandler, OnSimulcamViewportClicked)
 	SLATE_EVENT(FSimulcamViewportInputKeyEventHandler, OnSimulcamViewportInputKey)
+	SLATE_EVENT(FSimulcamViewportMarqueeSelectEventHandler, OnSimulcamViewportMarqueeSelect)
 	SLATE_ATTRIBUTE(bool, WithZoom)
 	SLATE_ATTRIBUTE(bool, WithPan)
 	SLATE_END_ARGS()
@@ -58,6 +62,9 @@ public:
 	/** Called when the viewport receives input key presses*/
 	bool OnViewportInputKey(const FKey& Key, const EInputEvent& Event);
 
+	/** Called when the viewport receives a marquee select event */
+	void OnMarqueeSelect(FVector2D StartPosition, FVector2D EndPosition) { OnSimulcamViewportMarqueeSelect.ExecuteIfBound(StartPosition, EndPosition); }
+
 private:
 
 	/** Delegate to be executed when the viewport area is clicked */
@@ -66,6 +73,9 @@ private:
 	/** Delegate to be executed when the viewport receives input key presses */
 	FSimulcamViewportInputKeyEventHandler OnSimulcamViewportInputKey;
 
+	/** Delegate to be executed when the user performs a marquee selection on the viewport area */
+	FSimulcamViewportMarqueeSelectEventHandler OnSimulcamViewportMarqueeSelect;
+
 	TSharedPtr<SSimulcamEditorViewport> TextureViewport;
 
 	TStrongObjectPtr<UTexture> Texture;
diff --git a/Engine/Plugins/VirtualProduction/CameraCalibration/Source/CameraCalibrationEditor/Private/AssetEditor/Viewport/SimulcamEditorViewportClient.cpp b/Engine/Plugins/VirtualProduction/CameraCalibration/Source/CameraCalibrationEditor/Private/AssetEditor/Viewport/SimulcamEditorViewportClient.cpp
--- a/Engine/Plugins/VirtualProduction/CameraCalibration/Source/CameraCalibrationEditor/Private/AssetEditor/Viewport/SimulcamEditorViewportClient.cpp
+++ b/Engine/Plugins/VirtualProduction/CameraCalibration/Source/CameraCalibrationEditor/Private/AssetEditor/Viewport/SimulcamEditorViewportClient.cpp
@@ -64,6 +64,13 @@ void FSimulcamEditorViewportClient::Draw(FViewport* Viewport, FCanvas* Canvas)
 		TileItem.BatchedElementParameters = BatchedElementParameters;
 		Canvas->DrawItem(TileItem);
 	}
+
+	// If the user is current doing a marquee select, draw the marquee selection box
+	if (bIsMarqueeSelecting)
+	{
+		FCanvasBoxItem BoxItem(SelectionStartCanvas, SelectionBoxSize);
+		Canvas->DrawItem(BoxItem);
+	}
 }
 
 void FSimulcamEditorViewportClient::MouseMove(FViewport* Viewport, int32 X, int32 Y)
@@ -121,6 +128,9 @@ bool FSimulcamEditorViewportClient::InputKey(const FInputKeyEventArgs& InEventAr
 			return false;
 		}
 
+		const bool bIsCtrlDown = FSlateApplication::Get().GetPlatformApplication()->GetModifierKeys().IsControlDown();
+		const bool bIsAltDown = FSlateApplication::Get().GetPlatformApplication()->GetModifierKeys().IsAltDown();
+
 		if (InEventArgs.Key == EKeys::LeftMouseButton || InEventArgs.Key == EKeys::MiddleMouseButton || InEventArgs.Key == EKeys::RightMouseButton)
 		{
 			const FGeometry& MyGeometry = SimulcamEditorViewportWeakPtr.Pin()->GetTickSpaceGeometry();
@@ -129,31 +139,35 @@ bool FSimulcamEditorViewportClient::InputKey(const FInputKeyEventArgs& InEventAr
 			if (LocalMouse.ComponentwiseAllGreaterOrEqual(FVector2D(0, 0)) &&
 				LocalMouse.ComponentwiseAllLessThan(MyGeometry.GetAbsoluteSize()))
 			{
-				if (UTexture* Texture = SimulcamViewportWeakPtr.Pin()->GetTexture())
+				// create fake geometry and mouseposition
+				const FVector2D FakeMousePosition = GetTexturePosition();
+				const FIntPoint TextureSize = GetTextureSize();
+
+				// check for meaningful position
+				if (FakeMousePosition.X >= 0 && FakeMousePosition.Y >= 0 && FakeMousePosition.X < TextureSize.X && FakeMousePosition.Y < TextureSize.Y)
 				{
-					if (Texture->GetResource())
+					if (bIsCtrlDown && bIsAltDown && !bIsMarqueeSelecting)
 					{
-						uint32 TextureWidth = Texture->GetResource()->GetSizeX();
-						uint32 TextureHeight = Texture->GetResource()->GetSizeY();
-
-						// create fake geometry and mouseposition
-						const FVector2D FakeMousePosition = GetTexturePosition();
-
-						// check for meaningful position
-						if (FakeMousePosition.X >= 0 && FakeMousePosition.Y >= 0 && FakeMousePosition.X < TextureWidth && FakeMousePosition.Y < TextureHeight)
-						{
-							const FGeometry FakeGeometry = FGeometry::MakeRoot(FVector2D(TextureWidth, TextureHeight), FSlateLayoutTransform());
-							FPointerEvent PointerEvent(
-								FSlateApplicationBase::CursorPointerIndex,
-								FakeMousePosition,
-								FakeMousePosition,
-								TSet<FKey>(),
-								InEventArgs.Key,
-								0,
-								FSlateApplication::Get().GetPlatformApplication()->GetModifierKeys());
-
-							SimulcamEditorViewportWeakPtr.Pin()->OnViewportClicked(FakeGeometry, PointerEvent);
-						}
+						// The user is initiating a marquee select
+						bIsMarqueeSelecting = true;
+						SelectionStartCanvas = LocalMouse;
+						SelectionStartTexture = FakeMousePosition;
+						SelectionBoxSize = FVector2D(0);
+					}
+					else 
+					{
+						// The user is performing some other mouse click event
+						const FGeometry FakeGeometry = FGeometry::MakeRoot(FVector2D(TextureSize), FSlateLayoutTransform());
+						FPointerEvent PointerEvent(
+							FSlateApplicationBase::CursorPointerIndex,
+							FakeMousePosition,
+							FakeMousePosition,
+							TSet<FKey>(),
+							InEventArgs.Key,
+							0,
+							FSlateApplication::Get().GetPlatformApplication()->GetModifierKeys());
+
+						SimulcamEditorViewportWeakPtr.Pin()->OnViewportClicked(FakeGeometry, PointerEvent);
 					}
 				}
 			}
@@ -161,8 +175,6 @@ bool FSimulcamEditorViewportClient::InputKey(const FInputKeyEventArgs& InEventAr
 
 		if (bWithZoom)
 		{
-			const bool bIsCtrlDown = FSlateApplication::Get().GetPlatformApplication()->GetModifierKeys().IsControlDown();
-
 			if (InEventArgs.Key == EKeys::MouseScrollUp || (InEventArgs.Key == EKeys::Add && bIsCtrlDown))
 			{
 				ZoomOnPoint(InEventArgs.Viewport, MousePosition, [this] {SimulcamEditorViewportWeakPtr.Pin()->ZoomIn(); });
@@ -194,6 +206,26 @@ bool FSimulcamEditorViewportClient::InputKey(const FInputKeyEventArgs& InEventAr
 	}
 	else if (InEventArgs.Event == IE_Released)
 	{
+		if (bIsMarqueeSelecting)
+		{
+			// The user has finished the marquee select
+			bIsMarqueeSelecting = false;
+
+			// Calculate where the selection ended in texture coordinates based on the latest selection box size
+			const double CurrentZoom = SimulcamEditorViewportWeakPtr.Pin()->GetCustomZoomLevel();
+			FVector2D SelectionEndTexture = (SelectionStartTexture + (SelectionBoxSize / CurrentZoom));
+
+			// Clamp end position to the texture size
+			const FIntPoint TextureSize = GetTextureSize();
+			SelectionEndTexture.X = FMath::Clamp(SelectionEndTexture.X, 0, TextureSize.X);
+			SelectionEndTexture.Y = FMath::Clamp(SelectionEndTexture.Y, 0, TextureSize.Y);
+
+			const FVector2D TopLeftPoint = FIntPoint(FMath::Min(SelectionStartTexture.X, SelectionEndTexture.X), FMath::Min(SelectionStartTexture.Y, SelectionEndTexture.Y));
+			const FVector2D BottomRightPoint = FIntPoint(FMath::Max(SelectionStartTexture.X, SelectionEndTexture.X), FMath::Max(SelectionStartTexture.Y, SelectionEndTexture.Y));
+
+			SimulcamViewportWeakPtr.Pin()->OnMarqueeSelect(TopLeftPoint, BottomRightPoint);
+		}
+
 		return SimulcamViewportWeakPtr.Pin()->OnViewportInputKey(InEventArgs.Key, InEventArgs.Event);
 	}
 	else if (InEventArgs.Event == IE_Repeat)
@@ -298,12 +330,33 @@ bool FSimulcamEditorViewportClient::InputAxis(FViewport* Viewport, FInputDeviceI
 				CurrentTexturePosition.X = FMath::Clamp(CurrentTexturePosition.X, Viewport->GetSizeXY().X - TextureSize.X, 0.0f);
 			}
 		}
+
+		// Update the marquee selection box size based on the current mouse position
+		if (bIsMarqueeSelecting)
+		{
+			FVector2D LocalMouse = FVector2D(Viewport->GetMouseX(), Viewport->GetMouseY());
+			SelectionBoxSize = LocalMouse - SelectionStartCanvas;
+			ClampSelectionBoxSizeToTextureSize();
+		}
+
 		return true;
 	}
 
 	return false;
 }
 
+void FSimulcamEditorViewportClient::ClampSelectionBoxSizeToTextureSize()
+{
+	const FIntPoint TextureSize = GetTextureSize();
+	const double CurrentZoom = SimulcamEditorViewportWeakPtr.Pin()->GetCustomZoomLevel();
+
+	const FVector2D MinBoxSize = (FVector2D(0.0f, 0.0f) - SelectionStartTexture) * CurrentZoom;
+	const FVector2D MaxBoxSize = (FVector2D(TextureSize) - SelectionStartTexture) * CurrentZoom;
+
+	SelectionBoxSize.X = FMath::Clamp(SelectionBoxSize.X, MinBoxSize.X, MaxBoxSize.X);
+	SelectionBoxSize.Y = FMath::Clamp(SelectionBoxSize.Y, MinBoxSize.Y, MaxBoxSize.Y);
+}
+
 bool FSimulcamEditorViewportClient::ShouldUseMousePanning(FViewport* Viewport) const
 {
 	return bWithPan && SimulcamViewportWeakPtr.IsValid() && SimulcamViewportWeakPtr.Pin()->GetTexture() && Viewport->KeyState(EKeys::RightMouseButton);
@@ -393,3 +446,19 @@ FVector2D FSimulcamEditorViewportClient::GetTexturePosition() const
 	const double CurrentZoom = SimulcamEditorViewportWeakPtr.Pin()->GetCustomZoomLevel();
 	return FVector2D((MousePosition.X - CurrentTexturePosition.X) / CurrentZoom, (MousePosition.Y - CurrentTexturePosition.Y) / CurrentZoom);
 }
+
+FIntPoint FSimulcamEditorViewportClient::GetTextureSize() const 
+{
+	FIntPoint TextureSize = FIntPoint(1);
+
+	if (UTexture* Texture = SimulcamViewportWeakPtr.Pin()->GetTexture())
+	{
+		if (Texture->GetResource())
+		{
+			TextureSize.X = Texture->GetResource()->GetSizeX();
+			TextureSize.Y = Texture->GetResource()->GetSizeY();
+		}
+	}
+
+	return TextureSize;
+}
diff --git a/Engine/Plugins/VirtualProduction/CameraCalibration/Source/CameraCalibrationEditor/Private/AssetEditor/Viewport/SimulcamEditorViewportClient.h b/Engine/Plugins/VirtualProduction/CameraCalibration/Source/CameraCalibrationEditor/Private/AssetEditor/Viewport/SimulcamEditorViewportClient.h
--- a/Engine/Plugins/VirtualProduction/CameraCalibration/Source/CameraCalibrationEditor/Private/AssetEditor/Viewport/SimulcamEditorViewportClient.h
+++ b/Engine/Plugins/VirtualProduction/CameraCalibration/Source/CameraCalibrationEditor/Private/AssetEditor/Viewport/SimulcamEditorViewportClient.h
@@ -42,7 +42,8 @@ public:
 
 private:
 	FVector2D GetTexturePosition() const;
-	
+	FIntPoint GetTextureSize() const;
+
 	/** TRUE if right clicking and dragging for panning a texture 2D */
 	bool ShouldUseMousePanning(FViewport* Viewport) const;
 
@@ -58,6 +59,9 @@ private:
 	/** Zoom Out trying to (smoothly) re-fit the image if there is space on borders */
 	void ZoomTowardsFit(FViewport* Viewport);
 
+	/** Clamp the selection box size to only be large as the texture size */
+	void ClampSelectionBoxSizeToTextureSize();
+
 	TWeakPtr<SSimulcamViewport> SimulcamViewportWeakPtr;
 	TWeakPtr<SSimulcamEditorViewport> SimulcamEditorViewportWeakPtr;
 
@@ -69,4 +73,9 @@ private:
 
 	bool bWithZoom;
 	bool bWithPan;
+
+	bool bIsMarqueeSelecting = false;
+	FVector2D SelectionStartTexture = FVector2D(0);
+	FVector2D SelectionStartCanvas = FVector2D(0);
+	FVector2D SelectionBoxSize = FVector2D(0);
 };
diff --git a/Engine/Plugins/VirtualProduction/CameraCalibrationCore/Source/CameraCalibrationCore/Public/CameraCalibrationStep.h b/Engine/Plugins/VirtualProduction/CameraCalibrationCore/Source/CameraCalibrationCore/Public/CameraCalibrationStep.h
--- a/Engine/Plugins/VirtualProduction/CameraCalibrationCore/Source/CameraCalibrationCore/Public/CameraCalibrationStep.h
+++ b/Engine/Plugins/VirtualProduction/CameraCalibrationCore/Source/CameraCalibrationCore/Public/CameraCalibrationStep.h
@@ -41,6 +41,9 @@ public:
 	/** Callback when key is pressed while viewport is focused. Returns false if the event was not handled. */
 	virtual bool OnViewportInputKey(const FKey& InKey, const EInputEvent& InEvent) { return false; };
 
+	/** Callback when the user performs a marquee select operation. Returns false if the event was not handled. */
+	virtual bool OnViewportMarqueeSelect(FVector2D StartPosition, FVector2D EndPosition) { return false; };
+
 	/** Returns the UI of this camera calibration step. Expected to only be called once */
 	virtual TSharedRef<SWidget> BuildUI() { return SNew(SVerticalBox); };
 
