CameraCalibration: Add marquee select support to the lens distortion tool to send subregions of the media image to opencv for checkerboard detection

diff --git a/Engine/Plugins/Runtime/OpenCV/Source/OpenCVHelper/Private/OpenCVHelper.cpp b/Engine/Plugins/Runtime/OpenCV/Source/OpenCVHelper/Private/OpenCVHelper.cpp
--- a/Engine/Plugins/Runtime/OpenCV/Source/OpenCVHelper/Private/OpenCVHelper.cpp
+++ b/Engine/Plugins/Runtime/OpenCV/Source/OpenCVHelper/Private/OpenCVHelper.cpp
@@ -480,11 +480,33 @@ bool FOpenCVHelper::DrawArucoMarkers(const TArray<FArucoMarker>& Markers, UTextu
 }
 
 bool FOpenCVHelper::IdentifyCheckerboard(TArray<FColor>& Image, FIntPoint ImageSize, FIntPoint CheckerboardDimensions, TArray<FVector2f>& OutCorners)
+{
+	FIntRect FullImageRect = FIntRect(FIntPoint(0), ImageSize);
+	return IdentifyCheckerboard(Image, ImageSize, FullImageRect, CheckerboardDimensions, OutCorners);
+}
+
+bool FOpenCVHelper::IdentifyCheckerboard(TArray<FColor>& Image, FIntPoint ImageSize, FIntRect RegionOfInterest, FIntPoint CheckerboardDimensions, TArray<FVector2f>& OutCorners)
 {
 #if WITH_OPENCV
 	// Initialize an OpenCV matrix header to point at the input image data. 
 	// The format for FColor is B8G8R8A8, so the corresponding OpenCV format is CV_8UC4 (8-bit per channel, 4 channels)
-	const cv::Mat ImageMat = cv::Mat(ImageSize.Y, ImageSize.X, CV_8UC4, Image.GetData());
+	const cv::Mat WholeImageMat = cv::Mat(ImageSize.Y, ImageSize.X, CV_8UC4, Image.GetData());
+
+	// Sanitize the ROI to ensure that it lies completely within the bounds of the full size image
+	RegionOfInterest.Min.X = FMath::Clamp(RegionOfInterest.Min.X, 0, RegionOfInterest.Min.X);
+	RegionOfInterest.Min.Y = FMath::Clamp(RegionOfInterest.Min.Y, 0, RegionOfInterest.Min.Y);
+	RegionOfInterest.Max.X = FMath::Clamp(RegionOfInterest.Min.X, RegionOfInterest.Max.X, ImageSize.X);
+	RegionOfInterest.Max.Y = FMath::Clamp(RegionOfInterest.Min.Y, RegionOfInterest.Max.Y, ImageSize.Y);
+
+	if (RegionOfInterest.IsEmpty())
+	{
+		OutCorners.Empty();
+		return false;
+	}
+
+	// Create a header for the region of interest
+	const cv::Rect CvROI = cv::Rect(RegionOfInterest.Min.X, RegionOfInterest.Min.Y, RegionOfInterest.Size().X, RegionOfInterest.Size().Y);
+	const cv::Mat ImageMat = cv::Mat(WholeImageMat, CvROI);
 
 	// Convert the image to grayscale before attempting to detect checkerboard corners
 	cv::Mat GrayImage;
@@ -521,6 +543,12 @@ bool FOpenCVHelper::IdentifyCheckerboard(TArray<FColor>& Image, FIntPoint ImageS
 		Algo::Reverse(OutCorners);
 	}
 
+	// Adjust the detected corners to be relative to the full image, not the ROI
+	for (FVector2f& Corner : OutCorners)
+	{
+		Corner += RegionOfInterest.Min;
+	}
+
 	return true;
 #else
 	return false;
diff --git a/Engine/Plugins/Runtime/OpenCV/Source/OpenCVHelper/Public/OpenCVHelper.h b/Engine/Plugins/Runtime/OpenCV/Source/OpenCVHelper/Public/OpenCVHelper.h
--- a/Engine/Plugins/Runtime/OpenCV/Source/OpenCVHelper/Public/OpenCVHelper.h
+++ b/Engine/Plugins/Runtime/OpenCV/Source/OpenCVHelper/Public/OpenCVHelper.h
@@ -149,9 +149,12 @@ public:
 	/** Draw a debug view of the input aruco markers on top of the input texture */
 	static bool DrawArucoMarkers(const TArray<FArucoMarker>& Markers, UTexture2D* DebugTexture);
 
-	/** Identify a checkerboard pattern in the input image that with the given checkerboard dimensions (columns x rows), and output the 2D coordinates of the intersections between each checkerboard square */
+	/** Identify a checkerboard pattern in the input image with the given checkerboard dimensions (columns x rows), and output the 2D coordinates of the intersections between each checkerboard square */
 	static bool IdentifyCheckerboard(TArray<FColor>& Image, FIntPoint ImageSize, FIntPoint CheckerboardDimensions, TArray<FVector2f>& OutCorners);
 
+	/** Identify a checkerboard pattern in a region of interest of the input image with the given checkerboard dimensions (columns x rows), and output the 2D coordinates of the intersections between each checkerboard square */
+	static bool IdentifyCheckerboard(TArray<FColor>& Image, FIntPoint ImageSize, FIntRect RegionOfInterest, FIntPoint CheckerboardDimensions, TArray<FVector2f>& OutCorners);
+
 	/** Draw a debug view of the input checkerboard corners on top of the input texture */
 	static bool DrawCheckerboardCorners(const TArray<FVector2f>& Corners, FIntPoint CheckerboardDimensions, UTexture2D* DebugTexture);
 
diff --git a/Engine/Plugins/VirtualProduction/CameraCalibration/Source/CameraCalibrationEditor/Private/AssetEditor/LensDistortionTool.cpp b/Engine/Plugins/VirtualProduction/CameraCalibration/Source/CameraCalibrationEditor/Private/AssetEditor/LensDistortionTool.cpp
--- a/Engine/Plugins/VirtualProduction/CameraCalibration/Source/CameraCalibrationEditor/Private/AssetEditor/LensDistortionTool.cpp
+++ b/Engine/Plugins/VirtualProduction/CameraCalibration/Source/CameraCalibrationEditor/Private/AssetEditor/LensDistortionTool.cpp
@@ -154,6 +154,29 @@ bool ULensDistortionTool::OnViewportClicked(const FGeometry& MyGeometry, const F
 		return false;
 	}
 
+	return CaptureCalibrationData(MyGeometry, MouseEvent);
+}
+
+bool ULensDistortionTool::OnViewportMarqueeSelect(FVector2D StartPosition, FVector2D EndPosition)
+{
+	// Marquee select currently only supports providing an ROI for checkerboard detection
+	if (CaptureSettings.CalibrationPattern != ECalibrationPattern::Checkerboard)
+	{
+		return false;
+	}
+
+	FIntRect MarqueeSelection;
+	MarqueeSelection.Min = FIntPoint(FMath::Floor(StartPosition.X), FMath::Floor(StartPosition.Y));
+	MarqueeSelection.Max = FIntPoint(FMath::Floor(EndPosition.X), FMath::Floor(EndPosition.Y));
+
+	// These are unused by the checkerboard detection
+	const FGeometry MyGeometryUnused;
+	const FPointerEvent MouseEventUnused;
+	return CaptureCalibrationData(MyGeometryUnused, MouseEventUnused, MarqueeSelection);
+}
+
+bool ULensDistortionTool::CaptureCalibrationData(const FGeometry& MyGeometry, const FPointerEvent& MouseEvent, FIntRect RegionOfInterest)
+{
 	// Block user interaction with the simulcam viewport while an async calibration task is executing
 	if (CalibrationTask.IsValid())
 	{
@@ -195,7 +218,12 @@ bool ULensDistortionTool::OnViewportClicked(const FGeometry& MyGeometry, const F
 		bool bDetectionResult = false;
 		if (CaptureSettings.CalibrationPattern == ECalibrationPattern::Checkerboard)
 		{
-			bDetectionResult = DetectCheckerboardPattern(Pixels, ImageSize, NewRow);
+			// If no ROI was provided, use the full image size
+			if (RegionOfInterest.IsEmpty())
+			{
+				RegionOfInterest = FIntRect(FIntPoint(0), ImageSize);
+			}
+			bDetectionResult = DetectCheckerboardPattern(Pixels, ImageSize, RegionOfInterest, NewRow);
 		}
 		else if (CaptureSettings.CalibrationPattern == ECalibrationPattern::Aruco)
 		{
@@ -240,7 +268,7 @@ bool ULensDistortionTool::OnViewportClicked(const FGeometry& MyGeometry, const F
 	return true;
 }
 
-bool ULensDistortionTool::DetectCheckerboardPattern(TArray<FColor>& Pixels, FIntPoint Size, TSharedPtr<FCalibrationRow> OutRow)
+bool ULensDistortionTool::DetectCheckerboardPattern(TArray<FColor>& Pixels, FIntPoint Size, FIntRect RegionOfInterest, TSharedPtr<FCalibrationRow> OutRow)
 {
 	const FText ErrorTitle = LOCTEXT("CaptureError", "Capture Error");
 
@@ -257,10 +285,10 @@ bool ULensDistortionTool::DetectCheckerboardPattern(TArray<FColor>& Pixels, FInt
 	OutRow->CheckerboardDimensions = CheckerboardDimensions;
 
 	// Launch an async task to perform the opencv checkerboard detection to prevent the game thread from being blocked in the rare cases when detection takes a very long time
-	UE::Tasks::TTask<TArray<FVector2f>> DetectionTask = UE::Tasks::Launch(UE_SOURCE_LOCATION, [Pixels, Size, CheckerboardDimensions]() mutable
+	UE::Tasks::TTask<TArray<FVector2f>> DetectionTask = UE::Tasks::Launch(UE_SOURCE_LOCATION, [Pixels, Size, RegionOfInterest, CheckerboardDimensions]() mutable
 	{
 		TArray<FVector2f> Corners;
-		FOpenCVHelper::IdentifyCheckerboard(Pixels, Size, CheckerboardDimensions, Corners);
+		FOpenCVHelper::IdentifyCheckerboard(Pixels, Size, RegionOfInterest, CheckerboardDimensions, Corners);
 		return Corners;
 	});
 
diff --git a/Engine/Plugins/VirtualProduction/CameraCalibration/Source/CameraCalibrationEditor/Private/AssetEditor/LensDistortionTool.h b/Engine/Plugins/VirtualProduction/CameraCalibration/Source/CameraCalibrationEditor/Private/AssetEditor/LensDistortionTool.h
--- a/Engine/Plugins/VirtualProduction/CameraCalibration/Source/CameraCalibrationEditor/Private/AssetEditor/LensDistortionTool.h
+++ b/Engine/Plugins/VirtualProduction/CameraCalibration/Source/CameraCalibrationEditor/Private/AssetEditor/LensDistortionTool.h
@@ -181,6 +181,7 @@ public:
 	virtual void Shutdown() override;
 	virtual void Tick(float DeltaTime) override;
 	virtual bool OnViewportClicked(const FGeometry& MyGeometry, const FPointerEvent& MouseEvent) override;
+	virtual bool OnViewportMarqueeSelect(FVector2D StartPosition, FVector2D EndPosition) override;
 	virtual TSharedRef<SWidget> BuildUI() override;
 	virtual FName FriendlyName() const  override { return TEXT("Lens Distortion"); };
 	virtual bool DependsOnStep(UCameraCalibrationStep* Step) const override;
@@ -216,8 +217,11 @@ private:
 	/** Get the latest status from the lens distortion calibration. Returns true if the status has changed. */
 	bool GetCalibrationStatus(FText& StatusText) const;
 
+	/** Capture a calibration pattern (or raw points) from the media image and collect additional data needed for a full calibration row */
+	bool CaptureCalibrationData(const FGeometry& MyGeometry, const FPointerEvent& MouseEvent, FIntRect RegionOfInterest = FIntRect());
+
 	/** Detects a checkerboard pattern in the input array of pixels (supplied by the media image) and fills out the 2D and 3D points in the output calibration row */
-	bool DetectCheckerboardPattern(TArray<FColor>& Pixels, FIntPoint Size, TSharedPtr<FCalibrationRow> OutRow);
+	bool DetectCheckerboardPattern(TArray<FColor>& Pixels, FIntPoint Size, FIntRect RegionOfInterest, TSharedPtr<FCalibrationRow> OutRow);
 
 	/** Detects an aruco pattern in the input array of pixels (supplied by the media image) and fills out the 2D and 3D points in the output calibration row */
 	bool DetectArucoPattern(TArray<FColor>& Pixels, FIntPoint Size, TSharedPtr<FCalibrationRow> OutRow);
