CameraCalibration: Fix bug where user deletes a calibration row but the UI does not refresh the list view, giving the appearance that the delete did not work

diff --git a/Engine/Plugins/VirtualProduction/CameraCalibration/Source/CameraCalibrationEditor/Private/AssetEditor/SLensDistortionToolPanel.cpp b/Engine/Plugins/VirtualProduction/CameraCalibration/Source/CameraCalibrationEditor/Private/AssetEditor/SLensDistortionToolPanel.cpp
--- a/Engine/Plugins/VirtualProduction/CameraCalibration/Source/CameraCalibrationEditor/Private/AssetEditor/SLensDistortionToolPanel.cpp
+++ b/Engine/Plugins/VirtualProduction/CameraCalibration/Source/CameraCalibrationEditor/Private/AssetEditor/SLensDistortionToolPanel.cpp
@@ -409,6 +409,8 @@ FReply SLensDistortionToolPanel::OnDatasetRowKeyPressed(const FGeometry& Geometr
 			Tool->DeleteExportedRow(SelectedItem->Index);
 		}
 
+		RefreshListView();
+
 		Tool->RefreshCoverage();
 
 		return FReply::Handled();
