CameraCalibration: Fix bug where the checkerboard coverage overlay does not draw on the viewport when the user imports a previously collected dataset

diff --git a/Engine/Plugins/VirtualProduction/CameraCalibration/Source/CameraCalibrationEditor/Private/AssetEditor/LensDistortionTool.cpp b/Engine/Plugins/VirtualProduction/CameraCalibration/Source/CameraCalibrationEditor/Private/AssetEditor/LensDistortionTool.cpp
--- a/Engine/Plugins/VirtualProduction/CameraCalibration/Source/CameraCalibrationEditor/Private/AssetEditor/LensDistortionTool.cpp
+++ b/Engine/Plugins/VirtualProduction/CameraCalibration/Source/CameraCalibrationEditor/Private/AssetEditor/LensDistortionTool.cpp
@@ -1149,6 +1149,9 @@ void ULensDistortionTool::ImportCalibrationDataset()
 		DistortionWidget->RefreshListView();
 	}
 
+	// Redraw the coverage overlay for the imported data
+	RefreshCoverage();
+
 	// Set the current session's start date/time and row index to match what was just imported to support adding/deleting rows
 	SessionInfo.bIsActive = true;
 	SessionInfo.StartTime = ImportedSessionDateTime;
