CameraCalibration: Fix nodal offset being saved with wrong focus/zoom keys when solved with distortion

diff --git a/Engine/Plugins/VirtualProduction/CameraCalibration/Source/CameraCalibrationEditor/Private/AssetEditor/LensDistortionTool.cpp b/Engine/Plugins/VirtualProduction/CameraCalibration/Source/CameraCalibrationEditor/Private/AssetEditor/LensDistortionTool.cpp
--- a/Engine/Plugins/VirtualProduction/CameraCalibration/Source/CameraCalibrationEditor/Private/AssetEditor/LensDistortionTool.cpp
+++ b/Engine/Plugins/VirtualProduction/CameraCalibration/Source/CameraCalibrationEditor/Private/AssetEditor/LensDistortionTool.cpp
@@ -710,7 +710,7 @@ void ULensDistortionTool::SaveCalibrationResult()
 			FScopedTransaction Transaction(LOCTEXT("SaveNodalOffsetResult", "Save Calibrated Nodal Offset to Lens Asset"));
 			LensFile->Modify();
 
-			LensFile->AddNodalOffsetPoint(CalibrationResult.EvaluatedFocus, CalibrationResult.EvaluatedFocus, CalibrationResult.NodalOffset);
+			LensFile->AddNodalOffsetPoint(CalibrationResult.EvaluatedFocus, CalibrationResult.EvaluatedZoom, CalibrationResult.NodalOffset);
 		}
 	}
 
