CameraCalibration: Fix mouse position not updating in simulcam viewport when main canvas loses focus

diff --git a/Engine/Plugins/VirtualProduction/CameraCalibration/Source/CameraCalibrationEditor/Private/AssetEditor/Viewport/SimulcamEditorViewportClient.cpp b/Engine/Plugins/VirtualProduction/CameraCalibration/Source/CameraCalibrationEditor/Private/AssetEditor/Viewport/SimulcamEditorViewportClient.cpp
--- a/Engine/Plugins/VirtualProduction/CameraCalibration/Source/CameraCalibrationEditor/Private/AssetEditor/Viewport/SimulcamEditorViewportClient.cpp
+++ b/Engine/Plugins/VirtualProduction/CameraCalibration/Source/CameraCalibrationEditor/Private/AssetEditor/Viewport/SimulcamEditorViewportClient.cpp
@@ -131,10 +131,13 @@ bool FSimulcamEditorViewportClient::InputKey(const FInputKeyEventArgs& InEventAr
 		const bool bIsCtrlDown = FSlateApplication::Get().GetPlatformApplication()->GetModifierKeys().IsControlDown();
 		const bool bIsAltDown = FSlateApplication::Get().GetPlatformApplication()->GetModifierKeys().IsAltDown();
 
+		const FVector2D LocalMouse = FVector2D(InEventArgs.Viewport->GetMouseX(), InEventArgs.Viewport->GetMouseY());
+		MousePosition.X = FMath::Floor(LocalMouse.X);
+		MousePosition.Y = FMath::Floor(LocalMouse.Y);
+
 		if (InEventArgs.Key == EKeys::LeftMouseButton || InEventArgs.Key == EKeys::MiddleMouseButton || InEventArgs.Key == EKeys::RightMouseButton)
 		{
 			const FGeometry& MyGeometry = SimulcamEditorViewportWeakPtr.Pin()->GetTickSpaceGeometry();
-			const FVector2D LocalMouse = FVector2D(InEventArgs.Viewport->GetMouseX(), InEventArgs.Viewport->GetMouseY());
 			// check if we are under the viewport, otherwise the capture system will blindly trigger the PointerEvent
 			if (LocalMouse.ComponentwiseAllGreaterOrEqual(FVector2D(0, 0)) &&
 				LocalMouse.ComponentwiseAllLessThan(MyGeometry.GetAbsoluteSize()))
