CameraCalibration: Add editor telemetry

diff --git a/Engine/Plugins/VirtualProduction/CameraCalibration/Source/CameraCalibrationEditor/Private/AssetEditor/LensDistortionTool.cpp b/Engine/Plugins/VirtualProduction/CameraCalibration/Source/CameraCalibrationEditor/Private/AssetEditor/LensDistortionTool.cpp
--- a/Engine/Plugins/VirtualProduction/CameraCalibration/Source/CameraCalibrationEditor/Private/AssetEditor/LensDistortionTool.cpp
+++ b/Engine/Plugins/VirtualProduction/CameraCalibration/Source/CameraCalibrationEditor/Private/AssetEditor/LensDistortionTool.cpp
@@ -12,6 +12,7 @@
 #include "CameraCalibrationUtilsPrivate.h"
 #include "DesktopPlatformModule.h"
 #include "Dom/JsonObject.h"
+#include "EngineAnalytics.h"
 #include "Framework/Application/SlateApplication.h"
 #include "HAL/FileManager.h"
 #include "IDesktopPlatform.h"
@@ -36,6 +37,25 @@ namespace UE::CameraCalibration::Private::LensDistortionTool
 	static const FString Version(TEXT("Version"));
 }
 
+namespace UE::LensDistortionToolAnalytics
+{
+	void RecordEvent(FLensCaptureSettings CaptureSettings, FLensSolverSettings SolverSettings, int32 DatasetSize)
+	{
+		if (FEngineAnalytics::IsAvailable())
+		{
+			TArray<FAnalyticsEventAttribute> EventAttributes;
+
+			EventAttributes.Add(FAnalyticsEventAttribute(TEXT("Pattern"), *UEnum::GetDisplayValueAsText(CaptureSettings.CalibrationPattern).ToString()));
+			EventAttributes.Add(FAnalyticsEventAttribute(TEXT("IsCameraTracked"), FString::Printf(TEXT("%d"), CaptureSettings.bIsCameraTracked ? 1 : 0)));
+			EventAttributes.Add(FAnalyticsEventAttribute(TEXT("IsCalibratorTracked"), FString::Printf(TEXT("%d"), CaptureSettings.bIsCalibratorTracked ? 1 : 0)));
+			EventAttributes.Add(FAnalyticsEventAttribute(TEXT("SolveNodalOffset"), FString::Printf(TEXT("%d"), SolverSettings.bSolveNodalOffset ? 1 : 0)));
+			EventAttributes.Add(FAnalyticsEventAttribute(TEXT("DatasetSize"), FString::Printf(TEXT("%d"), DatasetSize)));
+
+			FEngineAnalytics::GetProvider().RecordEvent(TEXT("CameraCalibration.DistortionCalibrationStarted"), EventAttributes);
+		}
+	}
+}
+
 void ULensDistortionTool::Initialize(TWeakPtr<FCameraCalibrationStepsController> InCameraCalibrationStepController)
 {
 	WeakStepsController = InCameraCalibrationStepController;
@@ -534,6 +554,8 @@ void ULensDistortionTool::CalibrateLens()
 		return;
 	}
 
+	UE::LensDistortionToolAnalytics::RecordEvent(CaptureSettings, SolverSettings, Dataset.CalibrationRows.Num());
+
 	const FIntPoint ImageSize = LensFile->CameraFeedInfo.GetDimensions();
 
 	const double FocalLengthEstimateValue = SolverSettings.FocalLengthGuess.GetValue();
diff --git a/Engine/Plugins/VirtualProduction/CameraCalibration/Source/CameraCalibrationEditor/Private/AssetEditor/SLensDataAddPointDialog.cpp b/Engine/Plugins/VirtualProduction/CameraCalibration/Source/CameraCalibrationEditor/Private/AssetEditor/SLensDataAddPointDialog.cpp
--- a/Engine/Plugins/VirtualProduction/CameraCalibration/Source/CameraCalibrationEditor/Private/AssetEditor/SLensDataAddPointDialog.cpp
+++ b/Engine/Plugins/VirtualProduction/CameraCalibration/Source/CameraCalibrationEditor/Private/AssetEditor/SLensDataAddPointDialog.cpp
@@ -5,6 +5,7 @@
 #include "CameraCalibrationSettings.h"
 #include "CameraCalibrationToolkit.h"
 #include "Curves/RichCurve.h"
+#include "EngineAnalytics.h"
 #include "Framework/MultiBox/MultiBoxBuilder.h"
 #include "IDetailChildrenBuilder.h"
 #include "IDetailPropertyRow.h"
@@ -24,6 +25,21 @@
 
 #define LOCTEXT_NAMESPACE "LensDataAddPointDialog"
 
+namespace UE::LensAddDataAnalytics
+{
+	void RecordEvent(ELensDataCategory Category)
+	{
+		if (FEngineAnalytics::IsAvailable())
+		{
+			TArray<FAnalyticsEventAttribute> EventAttributes;
+
+			EventAttributes.Add(FAnalyticsEventAttribute(TEXT("Category"), *UEnum::GetDisplayValueAsText(Category).ToString()));
+
+			FEngineAnalytics::GetProvider().RecordEvent(TEXT("CameraCalibration.AddLensDataPoint"), EventAttributes);
+		}
+	}
+}
+
 /**
  * Instanced customization used to display the distortion parameters
  * data structure associated with selected model
@@ -236,6 +252,8 @@ FReply SLensDataAddPointDialog::OnAddDataPointClicked()
 
 	AddDataToLensFile();
 
+	UE::LensAddDataAnalytics::RecordEvent(SelectedCategory);
+
 	OnDataPointAdded.ExecuteIfBound();
 
 	CloseDialog();
diff --git a/Engine/Plugins/VirtualProduction/CameraCalibration/Source/CameraCalibrationEditor/Private/Calibrators/CameraNodalOffsetAlgoPoints.cpp b/Engine/Plugins/VirtualProduction/CameraCalibration/Source/CameraCalibrationEditor/Private/Calibrators/CameraNodalOffsetAlgoPoints.cpp
--- a/Engine/Plugins/VirtualProduction/CameraCalibration/Source/CameraCalibrationEditor/Private/Calibrators/CameraNodalOffsetAlgoPoints.cpp
+++ b/Engine/Plugins/VirtualProduction/CameraCalibration/Source/CameraCalibrationEditor/Private/Calibrators/CameraNodalOffsetAlgoPoints.cpp
@@ -13,6 +13,7 @@
 #include "DistortionRenderingUtils.h"
 #include "Dom/JsonObject.h"
 #include "Editor.h"
+#include "EngineAnalytics.h"
 #include "Framework/MultiBox/MultiBoxBuilder.h"
 #include "GameFramework/Actor.h"
 #include "Input/Events.h"
@@ -46,6 +47,22 @@ namespace UE::CameraCalibration::Private::NodalOffsetPointsExportFields
 	static const FString Version(TEXT("Version"));
 }
 
+namespace UE::NodalOffsetAlgoPointsAnalytics
+{
+	void RecordEvent(FString FriendlyName, int32 DatasetSize)
+	{
+		if (FEngineAnalytics::IsAvailable())
+		{
+			TArray<FAnalyticsEventAttribute> EventAttributes;
+
+			EventAttributes.Add(FAnalyticsEventAttribute(TEXT("Pattern"), *FriendlyName));
+			EventAttributes.Add(FAnalyticsEventAttribute(TEXT("DatasetSize"), FString::Printf(TEXT("%d"), DatasetSize)));
+
+			FEngineAnalytics::GetProvider().RecordEvent(TEXT("CameraCalibration.NodalOffsetCalibrationStarted"), EventAttributes);
+		}
+	}
+}
+
 namespace CameraNodalOffsetAlgoPoints
 {
 	class SCalibrationRowGenerator
@@ -991,6 +1008,8 @@ bool UCameraNodalOffsetAlgoPoints::GetNodalOffset(FNodalPointOffset& OutNodalOff
 		}
 	}
 
+	UE::NodalOffsetAlgoPointsAnalytics::RecordEvent(FriendlyName().ToString(), CalibrationRows.Num());
+
 	// Undistort the 2D points in each calibration row 
 	UndistortCalibrationRowPoints();
 
