Fixing a build issue when automated tests are disabled

diff --git a/Engine/Plugins/VirtualProduction/CameraCalibration/Source/CameraCalibrationEditor/Private/TestCameraCalibration.cpp b/Engine/Plugins/VirtualProduction/CameraCalibration/Source/CameraCalibrationEditor/Private/TestCameraCalibration.cpp
--- a/Engine/Plugins/VirtualProduction/CameraCalibration/Source/CameraCalibrationEditor/Private/TestCameraCalibration.cpp
+++ b/Engine/Plugins/VirtualProduction/CameraCalibration/Source/CameraCalibrationEditor/Private/TestCameraCalibration.cpp
@@ -883,8 +883,12 @@ bool FTestNodalOffset::RunTest(const FString& Parameters)
 	return true;
 }
 
+#endif // WITH_DEV_AUTOMATION_TESTS
+
 void UAutoCalibrationTest::RunTests(FString Filename, bool bLogVerboseTestDescription)
 {
+
+#if WITH_DEV_AUTOMATION_TESTS
 	if (FPaths::IsRelative(Filename))
 	{
 		Filename = FPaths::ProjectContentDir() / Filename;
@@ -897,6 +901,6 @@ void UAutoCalibrationTest::RunTests(FString Filename, bool bLogVerboseTestDescri
 	{
 		UE::Private::CameraCalibration::AutomatedTests::TestCameraCalibration(Test, bLogVerboseTestDescription);
 	}
-}
-
 #endif // WITH_DEV_AUTOMATION_TESTS
+
+}
