Lens Calibration: Removed dependency on Composure, reworked calibration tool rendering to use a level viewport client that renders out a scene render with media texture overlay.

#virtualized

diff --git a/Engine/Plugins/Media/MediaFrameworkUtilities/Source/MediaFrameworkUtilities/Private/Profile/MediaProfilePlaybackManager.cpp b/Engine/Plugins/Media/MediaFrameworkUtilities/Source/MediaFrameworkUtilities/Private/Profile/MediaProfilePlaybackManager.cpp
--- a/Engine/Plugins/Media/MediaFrameworkUtilities/Source/MediaFrameworkUtilities/Private/Profile/MediaProfilePlaybackManager.cpp
+++ b/Engine/Plugins/Media/MediaFrameworkUtilities/Source/MediaFrameworkUtilities/Private/Profile/MediaProfilePlaybackManager.cpp
@@ -66,7 +66,7 @@ UMediaTexture* UMediaProfilePlaybackManager::GetSourceMediaTexture(UMediaSource*
 	return GetSourceMediaTextureFromIndex(MediaSourceIndex);
 }
 
-UMediaTexture* UMediaProfilePlaybackManager::OpenSourceFromIndex(int32 InMediaSourceIndex)
+UMediaTexture* UMediaProfilePlaybackManager::OpenSourceFromIndex(int32 InMediaSourceIndex, void* InConsumer)
 {
 	if (!MediaSourceTextures.IsValidIndex(InMediaSourceIndex))
 	{
@@ -95,11 +95,12 @@ UMediaTexture* UMediaProfilePlaybackManager::OpenSourceFromIndex(int32 InMediaSo
 
 	MediaTexture->SetMediaPlayer(MediaSourcePlayer);
 	
-	if (MediaSourcePlayer->IsPlaying())
+	if (!MediaSourcePlayer->IsClosed())
 	{
-		// If the media source is already being played, simply return the media texture
+		// If the media source is already open, simply return the media texture
 		if (MediaSourcePlayer->GetUrl() == MediaSource->GetUrl())
 		{
+			RegisterMediaSourceConsumer(MediaSourcePlayer, InConsumer);
 			return MediaTexture;
 		}
 	}
@@ -110,6 +111,7 @@ UMediaTexture* UMediaProfilePlaybackManager::OpenSourceFromIndex(int32 InMediaSo
 	
 	if (MediaSourcePlayer->OpenSourceWithOptions(MediaSource, Options))
 	{
+		RegisterMediaSourceConsumer(MediaSourcePlayer, InConsumer);
 		return MediaTexture;
 	}
 	
@@ -117,7 +119,7 @@ UMediaTexture* UMediaProfilePlaybackManager::OpenSourceFromIndex(int32 InMediaSo
 	return nullptr;
 }
 
-UMediaTexture* UMediaProfilePlaybackManager::OpenSource(UMediaSource* InMediaSource)
+UMediaTexture* UMediaProfilePlaybackManager::OpenSource(UMediaSource* InMediaSource, void* InConsumer)
 {
 	UMediaProfile* MediaProfile = GetOwningMediaProfile();
 	if (!MediaProfile)
@@ -126,29 +128,37 @@ UMediaTexture* UMediaProfilePlaybackManager::OpenSource(UMediaSource* InMediaSou
 	}
 	
 	int32 MediaSourceIndex = MediaProfile->FindMediaSourceIndex(InMediaSource);
-	return OpenSourceFromIndex(MediaSourceIndex);
+	return OpenSourceFromIndex(MediaSourceIndex, InConsumer);
 }
 
-void UMediaProfilePlaybackManager::CloseSourceFromIndex(int32 InMediaSourceIndex, bool bDestroyMediaPlayer)
+void UMediaProfilePlaybackManager::CloseSourceFromIndex(int32 InMediaSourceIndex, const FCloseSourceArgs& InArgs)
 {
 	if (!MediaSourcePlayers.Contains(InMediaSourceIndex))
 	{
 		return;
 	}
 
-	MediaSourcePlayers[InMediaSourceIndex]->Close();
-	if (UMediaTexture* MediaTexture = GetSourceMediaTextureFromIndex(InMediaSourceIndex))
-	{
-		MediaTexture->SetMediaPlayer(nullptr);
-	}
+	UMediaPlayer* MediaPlayer = MediaSourcePlayers[InMediaSourceIndex];
+	UnregisterMediaSourceConsumer(MediaPlayer, InArgs.Consumer);
 	
-	if (bDestroyMediaPlayer)
+	const bool bStillHasConsumers = GetMediaSourceConsumerCount(MediaPlayer) > 0;
+	if ((!bStillHasConsumers && !InArgs.bSoftClose) || InArgs.bForceClose)
 	{
-		MediaSourcePlayers.Remove(InMediaSourceIndex);
+		MediaPlayer->Close();
+		if (UMediaTexture* MediaTexture = GetSourceMediaTextureFromIndex(InMediaSourceIndex))
+		{
+			MediaTexture->SetMediaPlayer(nullptr);
+		}
+	
+		if (InArgs.bDestroyMediaPlayer)
+		{
+			SourceConsumers.Remove(MediaSourcePlayers[InMediaSourceIndex]);
+			MediaSourcePlayers.Remove(InMediaSourceIndex);
+		}
 	}
 }
 
-void UMediaProfilePlaybackManager::CloseSource(UMediaSource* InMediaSource, bool bDestroyMediaPlayer)
+void UMediaProfilePlaybackManager::CloseSource(UMediaSource* InMediaSource, const FCloseSourceArgs& InArgs)
 {
 	UMediaProfile* MediaProfile = GetOwningMediaProfile();
 	if (!MediaProfile)
@@ -157,7 +167,7 @@ void UMediaProfilePlaybackManager::CloseSource(UMediaSource* InMediaSource, bool
 	}
 	
 	int32 MediaSourceIndex = MediaProfile->FindMediaSourceIndex(InMediaSource);
-	CloseSourceFromIndex(MediaSourceIndex, bDestroyMediaPlayer);
+	CloseSourceFromIndex(MediaSourceIndex, InArgs);
 }
 
 UMediaPlayer* UMediaProfilePlaybackManager::FindOrCreateMediaPlayerForIndex(int32 InMediaSourceIndex)
@@ -232,6 +242,10 @@ void UMediaProfilePlaybackManager::ChangeMediaSourceIndex(int32 InOriginalMediaS
 			MediaSourcePlayers.Add(InNewMediaSourceIndex, MediaPlayer);
 		}
 	}
+	else
+	{
+		SourceConsumers.Remove(MediaPlayer);
+	}
 
 	// Finally, refresh all media textures to have the correct media player attached
 	for (int32 Index = 0; Index < MediaSourceTextures.Num(); ++Index)
@@ -254,3 +268,39 @@ UMediaProfile* UMediaProfilePlaybackManager::GetOwningMediaProfile() const
 {
 	return GetTypedOuter<UMediaProfile>();
 }
+
+void UMediaProfilePlaybackManager::RegisterMediaSourceConsumer(UMediaPlayer* InMediaPlayer, void* InConsumer)
+{
+	if (InConsumer)
+	{
+		if (SourceConsumers.Contains(InMediaPlayer))
+		{
+			SourceConsumers[InMediaPlayer].Add(InConsumer);
+		}
+		else
+		{
+			SourceConsumers.Add(InMediaPlayer, TSet<void*> { InConsumer });
+		}
+	}
+}
+
+void UMediaProfilePlaybackManager::UnregisterMediaSourceConsumer(UMediaPlayer* InMediaPlayer, void* InConsumer)
+{
+	if (InConsumer)
+	{
+		if (SourceConsumers.Contains(InMediaPlayer))
+		{
+			SourceConsumers[InMediaPlayer].Remove(InConsumer);
+		}
+	}
+}
+
+int32 UMediaProfilePlaybackManager::GetMediaSourceConsumerCount(UMediaPlayer* InMediaPlayer) const
+{
+	if (SourceConsumers.Contains(InMediaPlayer))
+	{
+		return SourceConsumers[InMediaPlayer].Num();
+	}
+
+	return 0;
+}
diff --git a/Engine/Plugins/Media/MediaFrameworkUtilities/Source/MediaFrameworkUtilities/Public/Profile/MediaProfilePlaybackManager.h b/Engine/Plugins/Media/MediaFrameworkUtilities/Source/MediaFrameworkUtilities/Public/Profile/MediaProfilePlaybackManager.h
--- a/Engine/Plugins/Media/MediaFrameworkUtilities/Source/MediaFrameworkUtilities/Public/Profile/MediaProfilePlaybackManager.h
+++ b/Engine/Plugins/Media/MediaFrameworkUtilities/Source/MediaFrameworkUtilities/Public/Profile/MediaProfilePlaybackManager.h
@@ -38,30 +38,50 @@ public:
 	/**
 	 * Opens and begins playing the specified media source 
 	 * @param InMediaSourceIndex The index of the media source within the media profile
+	 * @param InConsumer If not null, this consumer will be added to the list of consumers using the open media source
 	 * @return The media texture that is outputting the media source, or null if no media source exists at the specified index
 	 */
-	UMediaTexture* OpenSourceFromIndex(int32 InMediaSourceIndex);
+	UMediaTexture* OpenSourceFromIndex(int32 InMediaSourceIndex, void* InConsumer = nullptr);
 
 	/**
 	 * Opens and begins playing the specified media source 
 	 * @param InMediaSource The media source within the media profile
+	 * @param InConsumer If not null, this consumer will be added to the list of consumers using the open media source
 	 * @return The media texture that is outputting the media source, or null if the media source is not a valid member of the media profile
 	 */
-	UMediaTexture* OpenSource(UMediaSource* InMediaSource);
+	UMediaTexture* OpenSource(UMediaSource* InMediaSource, void* InConsumer = nullptr);
 
+	/** Arguments that control what happens when a media source is closed */
+	struct FCloseSourceArgs
+	{
+		/** The consumer requesting the media source be closed */
+		void* Consumer = nullptr;
+
+		/** Indicates that the media player should be destroyed after closing */
+		bool bDestroyMediaPlayer = false;
+
+		/** Indicates that the media source should be closed even if it still has active consumers */
+		bool bForceClose = false;
+
+		/** Indicates that the media source should stay open even if there are no active consumers */
+		bool bSoftClose = false;
+
+		FCloseSourceArgs() { }
+	};
+	
 	/**
 	 * Closes the specified media source, and optionally destroys the media player that was playing the source
 	 * @param InMediaSourceIndex The index of the media source within the media profile
-	 * @param bDestroyMediaPlayer Optional parameter to destroy the media player playing the source or not
+	 * @param InArgs Arguments that control what happens when the media source is closed
 	 */
-	void CloseSourceFromIndex(int32 InMediaSourceIndex, bool bDestroyMediaPlayer = false);
+	void CloseSourceFromIndex(int32 InMediaSourceIndex, const FCloseSourceArgs& InArgs = FCloseSourceArgs());
 	
 	/**
 	 * Closes the specified media source, and optionally destroys the media player that was playing the source
 	 * @param InMediaSource The media source within the media profile
-	 * @param bDestroyMediaPlayer Optional parameter to destroy the media player playing the source or not
+	 * @param InArgs Arguments that control what happens when the media source is closed
 	 */
-	void CloseSource(UMediaSource* InMediaSource, bool bDestroyMediaPlayer = false);
+	void CloseSource(UMediaSource* InMediaSource, const FCloseSourceArgs& InArgs = FCloseSourceArgs());
 
 private:
 	/**
@@ -80,6 +100,15 @@ private:
 	
 	/** Gets the media profile that owns this playback manager */
 	UMediaProfile* GetOwningMediaProfile() const;
+
+	/** Registers a consumer with a specific media player */ 
+	void RegisterMediaSourceConsumer(UMediaPlayer* InMediaPlayer, void* InConsumer);
+
+	/** Unregisters a consumer with a sepcific media player */
+	void UnregisterMediaSourceConsumer(UMediaPlayer* InMediaPlayer, void* InConsumer);
+
+	/** Gets the number of consumers that have registered as using the specified media player */
+	int32 GetMediaSourceConsumerCount(UMediaPlayer* InMediaPlayer) const;
 	
 private:
 	/** A map of media players, where the hash key is the index of the media source the player is playing */
@@ -90,5 +119,8 @@ private:
 	UPROPERTY(Transient)
 	TArray<TSoftObjectPtr<UMediaTexture>> MediaSourceTextures;
 
+	/** Keeps track of registered consumers that are actively using an open media source */
+	TMap<TObjectPtr<UMediaPlayer>, TSet<void*>> SourceConsumers;
+	
 	friend class UMediaProfile;
 };
diff --git a/Engine/Plugins/Media/MediaFrameworkUtilities/Source/MediaFrameworkUtilitiesEditor/Private/AssetEditor/MediaProfileEditor.cpp b/Engine/Plugins/Media/MediaFrameworkUtilities/Source/MediaFrameworkUtilitiesEditor/Private/AssetEditor/MediaProfileEditor.cpp
--- a/Engine/Plugins/Media/MediaFrameworkUtilities/Source/MediaFrameworkUtilitiesEditor/Private/AssetEditor/MediaProfileEditor.cpp
+++ b/Engine/Plugins/Media/MediaFrameworkUtilities/Source/MediaFrameworkUtilitiesEditor/Private/AssetEditor/MediaProfileEditor.cpp
@@ -311,7 +311,10 @@ void FMediaProfileEditor::CloseAllMediaSources()
 {
 	for (int32 Index = 0; Index < MediaProfileBeingEdited->NumMediaSources(); ++Index)
 	{
-		MediaProfileBeingEdited->GetPlaybackManager()->CloseSourceFromIndex(Index);
+		UMediaProfilePlaybackManager::FCloseSourceArgs Args;
+		Args.Consumer = this;
+		
+		MediaProfileBeingEdited->GetPlaybackManager()->CloseSourceFromIndex(Index, Args);
 	}
 }
 
@@ -606,8 +609,12 @@ void FMediaProfileEditor::OnMediaItemDeleted(UClass* InMediaType, int32 InMediaI
 {
 	if (InMediaType->IsChildOf<UMediaSource>())
 	{
-		constexpr bool bDestroyMediaPlayer = true;
-		MediaProfileBeingEdited->GetPlaybackManager()->CloseSourceFromIndex(InMediaItemIndex, bDestroyMediaPlayer);
+		UMediaProfilePlaybackManager::FCloseSourceArgs Args;
+		Args.Consumer = this;
+		Args.bDestroyMediaPlayer = true;
+		Args.bForceClose = true;
+		
+		MediaProfileBeingEdited->GetPlaybackManager()->CloseSourceFromIndex(InMediaItemIndex, Args);
 	}
 	else if (InMediaType->IsChildOf<UMediaOutput>())
 	{
diff --git a/Engine/Plugins/Media/MediaFrameworkUtilities/Source/MediaFrameworkUtilitiesEditor/Private/AssetEditor/SMediaProfileMediaItemDisplay.cpp b/Engine/Plugins/Media/MediaFrameworkUtilities/Source/MediaFrameworkUtilitiesEditor/Private/AssetEditor/SMediaProfileMediaItemDisplay.cpp
--- a/Engine/Plugins/Media/MediaFrameworkUtilities/Source/MediaFrameworkUtilitiesEditor/Private/AssetEditor/SMediaProfileMediaItemDisplay.cpp
+++ b/Engine/Plugins/Media/MediaFrameworkUtilities/Source/MediaFrameworkUtilitiesEditor/Private/AssetEditor/SMediaProfileMediaItemDisplay.cpp
@@ -81,7 +81,8 @@ void SMediaProfileMediaSourceDisplay::ConfigureMediaImage()
 		return;
 	}
 	
-	UMediaTexture* MediaTexture = MediaProfile->GetPlaybackManager()->OpenSourceFromIndex(MediaItemIndex);
+	// Open the media source, registering the media profile editor as a consumer
+	UMediaTexture* MediaTexture = MediaProfile->GetPlaybackManager()->OpenSourceFromIndex(MediaItemIndex, MediaProfileEditor.Pin().Get());
 
 	if (!MediaTexture)
 	{
diff --git a/Engine/Plugins/VirtualProduction/CameraCalibration/CameraCalibration.uplugin b/Engine/Plugins/VirtualProduction/CameraCalibration/CameraCalibration.uplugin
--- a/Engine/Plugins/VirtualProduction/CameraCalibration/CameraCalibration.uplugin
+++ b/Engine/Plugins/VirtualProduction/CameraCalibration/CameraCalibration.uplugin
@@ -25,10 +25,6 @@
 			"Name": "CameraCalibrationCore",
 			"Enabled": true
 		},
-		{
-			"Name": "Composure",
-			"Enabled": true
-		},
 		{
 			"Name": "LensComponent",
 			"Enabled": true
diff --git a/Engine/Plugins/VirtualProduction/CameraCalibration/Source/CameraCalibrationEditor/CameraCalibrationEditor.Build.cs b/Engine/Plugins/VirtualProduction/CameraCalibration/Source/CameraCalibrationEditor/CameraCalibrationEditor.Build.cs
--- a/Engine/Plugins/VirtualProduction/CameraCalibration/Source/CameraCalibrationEditor/CameraCalibrationEditor.Build.cs
+++ b/Engine/Plugins/VirtualProduction/CameraCalibration/Source/CameraCalibrationEditor/CameraCalibrationEditor.Build.cs
@@ -16,8 +16,6 @@ namespace UnrealBuildTool.Rules
 					"AssetTools",
 					"CameraCalibrationCore",
 					"CinematicCamera",
-					"Composure",
-					"ComposureLayersEditor",
 					"ContentBrowserData",
 					"Core",
 					"CoreUObject",
diff --git a/Engine/Plugins/VirtualProduction/CameraCalibration/Source/CameraCalibrationEditor/Private/AssetEditor/CameraCalibrationStepsController.cpp b/Engine/Plugins/VirtualProduction/CameraCalibration/Source/CameraCalibrationEditor/Private/AssetEditor/CameraCalibrationStepsController.cpp
--- a/Engine/Plugins/VirtualProduction/CameraCalibration/Source/CameraCalibrationEditor/Private/AssetEditor/CameraCalibrationStepsController.cpp
+++ b/Engine/Plugins/VirtualProduction/CameraCalibration/Source/CameraCalibrationEditor/Private/AssetEditor/CameraCalibrationStepsController.cpp
@@ -2,7 +2,6 @@
 
 #include "CameraCalibrationStepsController.h"
 
-#include "AssetRegistry/AssetData.h"
 #include "Camera/CameraActor.h"
 #include "CameraCalibrationEditorLog.h"
 #include "CameraCalibrationSettings.h"
@@ -11,23 +10,14 @@
 #include "CameraCalibrationToolkit.h"
 #include "CameraCalibrationTypes.h"
 #include "CineCameraActor.h"
-#include "CompElementEditorModule.h"
-#include "Components/SceneCaptureComponent2D.h"
-#include "CompositingCaptureBase.h"
-#include "CompositingElement.h"
 #include "ContentBrowserDataSubsystem.h"
-#include "CompositingElements/CompositingElementInputs.h"
-#include "CompositingElements/CompositingElementOutputs.h"
-#include "CompositingElements/CompositingElementPasses.h"
-#include "CompositingElements/CompositingElementTransforms.h"
 #include "Containers/Ticker.h"
 #include "Editor.h"
-#include "Engine/TextureRenderTarget2D.h"
 #include "Engine/UserDefinedEnum.h"
 #include "EngineUtils.h"
-#include "ICompElementManager.h"
 #include "IContentBrowserDataModule.h"
 #include "Input/Events.h"
+#include "Engine/Canvas.h"
 #include "Kismet/KismetRenderingLibrary.h"
 #include "LensComponent.h"
 #include "LensFile.h"
@@ -38,16 +28,15 @@
 #include "MediaSource.h"
 #include "MediaTexture.h"
 #include "Misc/MessageDialog.h"
-#include "Modules/ModuleManager.h"
 #include "Profile/IMediaProfileManager.h"
 #include "Profile/MediaProfile.h"
 #include "SCameraCalibrationSteps.h"
 #include "TextureResource.h"
 #include "TimeSynchronizableMediaSource.h"
 #include "UnrealClient.h"
+#include "Profile/MediaProfilePlaybackManager.h"
 #include "UObject/ObjectMacros.h"
 #include "UObject/UObjectGlobals.h"
-#include "UObject/UnrealType.h"
 #include "Widgets/SWidget.h"
 
 
@@ -55,20 +44,6 @@
 
 namespace CameraCalibrationStepsController
 {
-	/** Returns the ICompElementManager object */
-	ICompElementManager* GetCompElementManager()
-	{
-		static ICompElementEditorModule* EditorModule = FModuleManager::Get().GetModulePtr<ICompElementEditorModule>("ComposureLayersEditor");
-
-		TSharedPtr<ICompElementManager> CompElementManager;
-
-		if (EditorModule && (CompElementManager = EditorModule->GetCompElementManager()).IsValid())
-		{
-			return CompElementManager.Get();
-		}
-		return nullptr;
-	}
-
 	/** Gets all assets of type TObject from the content browser */
 	template<typename TObject>
 	TArray<TSoftObjectPtr<TObject>> GetUObjectAssets()
@@ -212,159 +187,57 @@ bool FCameraCalibrationStepsController::OnTick(float DeltaTime)
 		LensFileEvaluationInputs = LensComponent->GetLensFileEvaluationInputs();
 	}
 
-	// Update the output resolution of the comp to match the resolution of the media source (if it has changed)
-	const bool bCompResized = UpdateCompResolution();
-
-	// Update the output resolution of the CG Layer if either the Comp was resized or if the camera's aspect ratio has changed
-	UpdateCGResolution(bCompResized);
-
-	// Tick each of the calibration steps
-	for (TStrongObjectPtr<UCameraCalibrationStep>& Step : CalibrationSteps)
-	{
-		if (Step.IsValid())
-		{
-			Step->Tick(DeltaTime);
-		}
-	}
-
-	return true;
-}
-
-bool FCameraCalibrationStepsController::UpdateCompResolution()
-{
-	const UMediaPlayer* const MediaPlayerPtr = MediaPlayer.Get();
-
-	if (!MediaPlayerPtr)
-	{
-		return false;
-	}
-
-	const FIntPoint MediaDimensions = MediaPlayerPtr->GetVideoTrackDimensions(INDEX_NONE, INDEX_NONE);
-	const FIntPoint CompRenderResolution = GetCompRenderResolution();
-
-	bool bCompResized = false;
-
-	// If no track was found, the dimensions might be (0, 0)
-	if (MediaDimensions.X != 0 && MediaDimensions.Y != 0)
+	const FIntPoint MediaOverlayResolution = GetMediaOverlayResolution();
+	bool bMediaOverlayResolutionChanged = false;
+	if (MediaOverlayResolution != LastMediaOverlayResolution)
 	{
-		if (MediaDimensions != CompRenderResolution)
-		{
-			// Resize the output resolution of the top-level comp to match the incoming media dimensions
-			if (ACompositingElement* CompPtr = Comp.Get())
-			{
-				Comp->SetRenderResolution(MediaDimensions);
-			}
-
-			// Resize the output render targets to match the incoming media dimensions
-			if (UTextureRenderTarget2D* RenderTargetPtr = RenderTarget.Get())
-			{
-				RenderTargetPtr->ResizeTarget(MediaDimensions.X, MediaDimensions.Y);
-			}
-
-			if (UTextureRenderTarget2D* MediaPlateRenderTargetPtr = MediaPlateRenderTarget.Get())
-			{
-				MediaPlateRenderTargetPtr->ResizeTarget(MediaDimensions.X, MediaDimensions.Y);
-			}
-
-			if (UTextureRenderTarget2D* ToolOverlayRenderTargetPtr = ToolOverlayRenderTarget.Get())
-			{
-				ToolOverlayRenderTargetPtr->ResizeTarget(MediaDimensions.X, MediaDimensions.Y);
-			}
-
-			if (UTextureRenderTarget2D* UserOverlayRenderTargetPtr = UserOverlayRenderTarget.Get())
-			{
-				UserOverlayRenderTargetPtr->ResizeTarget(MediaDimensions.X, MediaDimensions.Y);
-			}
-
-			bCompResized = true;
-		}
+		LastMediaOverlayResolution = MediaOverlayResolution;
+		bMediaOverlayResolutionChanged = true;
 	}
-
+	
 	// Update the LensFile SimulcamInfo if the Media Dimensions have changed
 	if (ULensFile* const LensFilePtr = GetLensFile())
 	{
-		if (LensFilePtr->SimulcamInfo.MediaResolution != MediaDimensions)
-		{
-			LensFilePtr->SimulcamInfo.MediaResolution = MediaDimensions;
-			LensFilePtr->SimulcamInfo.MediaPlateAspectRatio = (MediaDimensions.Y != 0) ? (float)MediaDimensions.X / (float)MediaDimensions.Y : 0.0f;
-		}
-	}
-
-	return bCompResized;
-}
-
-void FCameraCalibrationStepsController::UpdateCGResolution(bool bCompResized)
-{
-	const UCineCameraComponent* const CineCameraComponentPtr = CineCameraComponent.Get();
-	const ACompositingElement* const CompPtr = Comp.Get();
-	ACompositingElement* const CGLayerPtr = CGLayer.Get();
-	ULensFile* const LensFilePtr = GetLensFile();
-
-	if (!CineCameraComponentPtr || !CompPtr || !CGLayerPtr || !LensFilePtr)
-	{
-		return;
-	}
-
-	// There are two cases where the CG Layer's output resolution might need to be resized:
-	// 1) If the aspect ratio of the camera source's filmback changed, to avoid any stretching of the CG
-	// 2) If the Comp's output resolution changed, to ensure that the CG layer's resolution is not larger than the Comp, and not unnecessarily smaller
-
-	const FIntPoint CGResolution = CGLayerPtr->GetRenderResolution();
-	const float CGLayerAspectRatio = (CGResolution.Y != 0) ? (float)CGResolution.X / (float)CGResolution.Y : 0.0f;
-	const float FilmbackAspectRatio = CineCameraComponentPtr->AspectRatio;
-
-	// Early-out if the filmback is somehow 0.0f (which shouldn't be possible)
-	if (FMath::IsNearlyEqual(FilmbackAspectRatio, 0.0f))
-	{
-		return;
-	}
-
-	constexpr float Tolerance = 0.01f;
-	const bool bFilmbackAspectRatioChanged = !FMath::IsNearlyEqual(FilmbackAspectRatio, CGLayerAspectRatio, Tolerance);
-
-	if (bFilmbackAspectRatioChanged || bCompResized)
-	{
-		const FIntPoint CompResolution = GetCompRenderResolution();
-		const float CompAspectRatio = (CompResolution.Y != 0) ? (float)CompResolution.X / (float)CompResolution.Y : 0.0f;
-
-		// Compute the camera feed dimensions that would fit the aspect ratio of the camera within the comp
-		// If the filmback aspect ratio is wider than the comp, then the width of the camera feed will equal the width of the comp, and the height will be scaled
-		// If the filmback aspect ratio is narrower than the comp, then the height of the camera feed will equal the height of the comp, and the width will be scaled
+		const FIntPoint CameraFeedResolution = LensFilePtr->CameraFeedInfo.GetDimensions();
 
-		FIntPoint NewCGResolution = CompResolution;
-
-		if (FilmbackAspectRatio > CompAspectRatio)
+		// Camera feed is not allowed to be larger than media resolution, so if this is detected, reset the camera feed dimensions
+		const bool bIsCameraFeedLargerThanMedia = CameraFeedResolution.X > MediaOverlayResolution.X || CameraFeedResolution.Y > MediaOverlayResolution.Y;
+		if (bIsCameraFeedLargerThanMedia || bMediaOverlayResolutionChanged)
 		{
-			NewCGResolution.Y = CompResolution.X / FilmbackAspectRatio;
+			ResetCameraFeedDimensions();
 		}
-		else if (CompAspectRatio > FilmbackAspectRatio)
+		
+		LensFilePtr->SimulcamInfo.MediaResolution = MediaOverlayResolution;
+		LensFilePtr->SimulcamInfo.MediaPlateAspectRatio = (MediaOverlayResolution.Y != 0) ? (float)MediaOverlayResolution.X / (float)MediaOverlayResolution.Y : 0.0f;
+		
+		if (CineCameraComponent.IsValid())
 		{
-			NewCGResolution.X = CompResolution.Y * FilmbackAspectRatio;
+			const float FilmbackAspectRatio = CineCameraComponent.Pin()->AspectRatio;
+			LensFilePtr->SimulcamInfo.CGLayerAspectRatio = FilmbackAspectRatio;
 		}
+	}
 
-		CGLayerPtr->SetRenderResolution(NewCGResolution);
-
-		LensFilePtr->SimulcamInfo.CGLayerAspectRatio = FilmbackAspectRatio;
-
-		if (!LensFilePtr->CameraFeedInfo.IsOverridden())
+	// Tick each of the calibration steps
+	for (TStrongObjectPtr<UCameraCalibrationStep>& Step : CalibrationSteps)
+	{
+		if (Step.IsValid())
 		{
-			constexpr bool bMarkAsOverridden = false;
-			LensFilePtr->CameraFeedInfo.SetDimensions(NewCGResolution, bMarkAsOverridden);
+			Step->Tick(DeltaTime);
 		}
-
-		UpdateAspectRatioCorrection();
 	}
+
+	return true;
 }
 
 void FCameraCalibrationStepsController::SetCameraFeedDimensionsFromMousePosition(FVector2D MousePosition)
 {
 	// Compute the size of the camera feed, using the MousePosition as one of its corners.
 	// We assume that the camera feed will always be centered in the media. If that assumption proves false, this math should be updated in the future.
-	const FIntPoint CompRenderResolution = GetCompRenderResolution();
+	const FIntPoint MediaOverlayResolution = GetMediaOverlayResolution();
 
 	FIntPoint CameraFeedDimensions = FIntPoint(0, 0);
-	CameraFeedDimensions.X = FMath::Abs((CompRenderResolution.X / 2) - FMath::Floor(MousePosition.X)) * 2;
-	CameraFeedDimensions.Y = FMath::Abs((CompRenderResolution.Y / 2) - FMath::Floor(MousePosition.Y)) * 2;
+	CameraFeedDimensions.X = FMath::Abs((MediaOverlayResolution.X / 2) - FMath::Floor(MousePosition.X)) * 2;
+	CameraFeedDimensions.Y = FMath::Abs((MediaOverlayResolution.Y / 2) - FMath::Floor(MousePosition.Y)) * 2;
 
 	// Early-out if the dimensions are invalid
 	if ((CameraFeedDimensions.X == 0) || (CameraFeedDimensions.Y == 0))
@@ -374,7 +247,7 @@ void FCameraCalibrationStepsController::SetCameraFeedDimensionsFromMousePosition
 
 	if (ULensFile* const LensFilePtr = GetLensFile())
 	{
-		// Compare the aspect ratio of the selected camera to the cg layer's aspect ratio to determine if the cg layer needs to be resized
+		// Compare the aspect ratio of the selected camera to the feed aspect ratio to determine if we should attempt to minimize the aspect ratio difference
 		if (CineCameraComponent.IsValid())
 		{
 			const float CameraFeedAspectRatio = CameraFeedDimensions.X / (float)CameraFeedDimensions.Y;
@@ -390,8 +263,6 @@ void FCameraCalibrationStepsController::SetCameraFeedDimensionsFromMousePosition
 			// Update the camera feed dimensions and mark as overridden because it was changed as the result of user interaction
 			constexpr bool bMarkAsOverridden = true;
 			LensFilePtr->CameraFeedInfo.SetDimensions(CameraFeedDimensions, bMarkAsOverridden);
-
-			UpdateAspectRatioCorrection();
 		}
 	}
 }
@@ -400,38 +271,51 @@ void FCameraCalibrationStepsController::SetCameraFeedDimensions(FIntPoint Dimens
 {
 	if (ULensFile* const LensFilePtr = GetLensFile())
 	{
-		LensFile->CameraFeedInfo.SetDimensions(Dimensions, bMarkAsOverridden);
-		UpdateAspectRatioCorrection();
+		LensFilePtr->CameraFeedInfo.SetDimensions(Dimensions, bMarkAsOverridden);
 	}
 }
 
-void FCameraCalibrationStepsController::UpdateAspectRatioCorrection()
+void FCameraCalibrationStepsController::ResetCameraFeedDimensions()
 {
-	const ULensFile* const LensFilePtr = LensFile.Get();
-	const ACompositingElement* const CompPtr = Comp.Get();
-	UCompositingElementMaterialPass* const MaterialPassPtr = MaterialPass.Get();
+	TStrongObjectPtr<ULensFile> PinnedLensFile = LensFile.Pin();
+	if (!PinnedLensFile.IsValid())
+	{
+		return;
+	}
+	
+	float CameraAspectRatio = 0.0f;
+	if (CineCameraComponent.IsValid())
+	{
+		CameraAspectRatio = CineCameraComponent->AspectRatio;
+	}
+	else if (Camera.IsValid())
+	{
+		CameraAspectRatio = Camera->GetCameraComponent()->AspectRatio;
+	}
 
-	if (!MaterialPassPtr || !CompPtr || !LensFilePtr)
+	if (FMath::IsNearlyZero(CameraAspectRatio))
 	{
 		return;
 	}
 
-	float AspectRatioCorrectionX = 1.0f;
-	float AspectRatioCorrectionY = 1.0f;
+	// Update the camera feed dimensions to match with the media resolution, corrected for the camera's aspect ratio
+	const FIntPoint MediaResolution = GetMediaOverlayResolution();
+	const float MediaAspectRatio = MediaResolution.Y != 0 ? MediaResolution.X / (float)MediaResolution.Y : 0.0f;
 
-	// If the camera feed dimensions are valid, use them and the dimensions of the media to set the aspect ratio correction material parameters
-	if (LensFilePtr->CameraFeedInfo.IsValid())
+	FIntPoint CameraFeedDimensions = MediaResolution;
+	if (CameraAspectRatio > MediaAspectRatio)
 	{
-		const FIntPoint CompRenderResolution = CompPtr->GetRenderResolution();
-		const FIntPoint CameraFeedDimensions = LensFilePtr->CameraFeedInfo.GetDimensions();
-
-		AspectRatioCorrectionX = CompRenderResolution.X / (float)CameraFeedDimensions.X;
-		AspectRatioCorrectionY = CompRenderResolution.Y / (float)CameraFeedDimensions.Y;
+		CameraFeedDimensions.Y = CameraFeedDimensions.X / CameraAspectRatio;
 	}
-
-	// If the camera feed dimensions are invalid, reset the correction parameters
-	MaterialPassPtr->Material.SetScalarOverride(TEXT("AspectRatioCorrection_H"), AspectRatioCorrectionX);
-	MaterialPassPtr->Material.SetScalarOverride(TEXT("AspectRatioCorrection_V"), AspectRatioCorrectionY);
+	else
+	{
+		CameraFeedDimensions.X = CameraFeedDimensions.Y * CameraAspectRatio;
+	}
+	
+	PinnedLensFile->SimulcamInfo.CGLayerAspectRatio = CameraAspectRatio;
+	
+	constexpr bool bMarkAsOverridden = false;
+	PinnedLensFile->CameraFeedInfo.SetDimensions(CameraFeedDimensions, bMarkAsOverridden);
 }
 
 void FCameraCalibrationStepsController::MinimizeAspectRatioError(FIntPoint& CameraFeedDimensions, float CameraAspectRatio)
@@ -477,61 +361,6 @@ void FCameraCalibrationStepsController::MinimizeAspectRatioError(FIntPoint& Came
 	}
 }
 
-TWeakObjectPtr<ACompositingElement> FCameraCalibrationStepsController::AddElement(ACompositingElement* Parent, FString& ClassPath, FString& ElementName) const
-{
-	UClass* ElementClass = StaticLoadClass(ACompositingElement::StaticClass(), nullptr, *ClassPath, nullptr, LOAD_None, nullptr);
-
-	if (!ElementClass)
-	{
-		return nullptr;
-	}
-
-	ICompElementManager* CompElementManager = CameraCalibrationStepsController::GetCompElementManager();
-
-	if (!CompElementManager)
-	{
-		return nullptr;
-	}
-
-	TWeakObjectPtr<ACompositingElement> Element = CompElementManager->CreateElement(
-		*ElementName, ElementClass, nullptr, EObjectFlags::RF_Transient | EObjectFlags::RF_DuplicateTransient);
-
-	if (!Element.IsValid())
-	{
-		return nullptr;
-	}
-
-	if (Parent)
-	{
-		// Attach layer to parent
-		Parent->AttachAsChildLayer(Element.Get());
-
-		// Place element under Parent to keep things organized.
-		Element->AttachToActor(Parent, FAttachmentTransformRules::SnapToTargetNotIncludingScale);
-	}
-
-	return Element;
-}
-
-
-ACompositingElement* FCameraCalibrationStepsController::FindElement(const FString& Name) const
-{
-	// The motivation of finding them is if the level was saved with these objects.
-	// A side effect is that we'll destroy them when closing the calibrator.
-
-	const EActorIteratorFlags Flags = EActorIteratorFlags::SkipPendingKill;
-
-	for (TActorIterator<AActor> It(GetWorld(), ACompositingElement::StaticClass(), Flags); It; ++It)
-	{
-		if (It->GetName() == Name)
-		{
-			return CastChecked<ACompositingElement>(*It);
-		}
-	}
-
-	return nullptr;
-}
-
 void FCameraCalibrationStepsController::Cleanup()
 {
 	for (TStrongObjectPtr<UCameraCalibrationStep>& Step : CalibrationSteps)
@@ -543,39 +372,17 @@ void FCameraCalibrationStepsController::Cleanup()
 		}
 	}
 
-	if (MediaPlayer.IsValid())
-	{
-		MediaPlayer->Close();
-		MediaPlayer.Reset();
-	}
-
-	if (Comp.IsValid()) // It may have been destroyed manually by the user.
-	{
-		Comp->Destroy();
-		Comp.Reset();
-	}
-
-	if (CGLayer.IsValid())
-	{
-		CGLayer->Destroy();
-		CGLayer.Reset();
-	}
+	ClearMedia();
 
-	if (MediaPlate.IsValid())
+	if (InternalMediaPlayer.IsValid())
 	{
-		MediaPlate->Destroy();
-		MediaPlate.Reset();
+		InternalMediaPlayer->Close();
+		InternalMediaPlayer.Reset();
 	}
-
-	MaterialPass.Reset();
+	
 	Camera.Reset();
-
-	MediaPlateRenderTarget.Reset();
-}
-
-FString FCameraCalibrationStepsController::NamespacedName(const FString&& Name) const
-{
-	return FString::Printf(TEXT("CameraCalib_%s_%s"), *LensFile->GetName(), *Name);
+	
+	IMediaProfileManager::Get().OnMediaProfileChanged().RemoveAll(this);
 }
 
 UWorld* FCameraCalibrationStepsController::GetWorld() const
@@ -583,29 +390,6 @@ UWorld* FCameraCalibrationStepsController::GetWorld() const
 	return GEditor ? GEditor->GetEditorWorldContext().World() : nullptr;
 }
 
-UTextureRenderTarget2D* FCameraCalibrationStepsController::GetRenderTarget() const
-{
-	return RenderTarget.Get();
-}
-
-FIntPoint FCameraCalibrationStepsController::GetCompRenderResolution() const
-{
-	if (ACompositingElement* CompPtr = Comp.Get())
-	{
-		return CompPtr->GetRenderResolution();
-	}
-	return FIntPoint(0, 0);
-}
-
-FIntPoint FCameraCalibrationStepsController::GetCGRenderResolution() const
-{
-	if (ACompositingElement* CGLayerPtr = CGLayer.Get())
-	{
-		return CGLayerPtr->GetRenderResolution();
-	}
-	return FIntPoint(0, 0);
-}
-
 FIntPoint FCameraCalibrationStepsController::GetCameraFeedSize() const
 {
 	if (ULensFile* LensFilePtr = GetLensFile())
@@ -616,322 +400,109 @@ FIntPoint FCameraCalibrationStepsController::GetCameraFeedSize() const
 	return FIntPoint(0, 0);
 }
 
-void FCameraCalibrationStepsController::CreateComp()
+void FCameraCalibrationStepsController::InitializeMediaPlayer()
 {
-	// Don't do anything if we already created it.
-	if (Comp.IsValid())
-	{
-		// Some items are exposed in the World Outliner, and they will get cleaned up and re-created
-		// when the configurator is closed and re-opened.
-
-		UE_LOG(LogCameraCalibrationEditor, Warning, TEXT("Tried to create Comp that already existed"));
-		return;
-	}
-
-	ICompElementManager* CompElementManager = CameraCalibrationStepsController::GetCompElementManager();
-
-	if (!CompElementManager)
-	{
-		UE_LOG(LogCameraCalibrationEditor, Warning, TEXT("Could not find ICompElementManager"));
-		return;
-	}
-
-	// Create Comp parent
-	{
-		const FString CompName = NamespacedName(TEXT("Comp"));
-
-		Comp = FindElement(CompName);
-
-		if (!Comp.IsValid())
-		{
-			Comp = CompElementManager->CreateElement(
-				*CompName, ACompositingElement::StaticClass(), nullptr, EObjectFlags::RF_Transient | EObjectFlags::RF_DuplicateTransient);
-		}
-
-		if (!Comp.IsValid())
-		{
-			UE_LOG(LogCameraCalibrationEditor, Warning, TEXT("Failed to create 'Comp'"));
-			Cleanup();
-			return;
-		}
-	}
-
-	// Convenience function to add comp elements
-	auto FindOrAddElementToComp = [&](
-		TWeakObjectPtr<ACompositingElement>& Element,
-		FString&& ElementClassPath,
-		FString&& ElementName) -> bool
-	{
-		if (Element.IsValid())
-		{
-			return true;
-		}
-
-		Element = FindElement(ElementName);
-
-		if (!Element.IsValid())
-		{
-			Element = AddElement(Comp.Get(), ElementClassPath, ElementName);
-		}
-
-		return Element.IsValid();
-	};
+	// Create media player and media texture and assign it to the media input of the media plate layer.
+	// Create MediaPlayer
 
-	if (!FindOrAddElementToComp(
-		CGLayer,
-		TEXT("/Composure/Blueprints/CompositingElements/BP_CgCaptureCompElement.BP_CgCaptureCompElement_C"),
-		NamespacedName(TEXT("CG"))))
-	{
-		UE_LOG(LogCameraCalibrationEditor, Warning, TEXT("Failed to create 'CG' Layer"));
-		Cleanup();
-		return;
-	}
+	// Using a strong reference prevents the MediaPlayer from going stale when the level is saved.
+	InternalMediaPlayer = TStrongObjectPtr<UMediaPlayer>(NewObject<UMediaPlayer>(
+		GetTransientPackage(),
+		MakeUniqueObjectName(GetTransientPackage(), UMediaPlayer::StaticClass())
+		));
 
-	if (!FindOrAddElementToComp(
-		MediaPlate,
-		TEXT("/Composure/Blueprints/CompositingElements/BP_MediaPlateCompElement.BP_MediaPlateCompElement_C"),
-		NamespacedName(TEXT("MediaPlate"))))
+	if (!InternalMediaPlayer.IsValid())
 	{
-		UE_LOG(LogCameraCalibrationEditor, Warning, TEXT("Failed to create 'MediaPlate'"));
+		UE_LOG(LogCameraCalibrationEditor, Warning, TEXT("Failed to create MediaPlayer"));
 		Cleanup();
 		return;
 	}
 
-	// The CG Layer's resolution will be tied to the aspect ratio of the camera's filmback, not the resolution of the top-level comp
-	CGLayer->ResolutionSource = EInheritedSourceType::Override;
-
-	// This updates the Composure panel view
-	CompElementManager->RefreshElementsList();
+	InternalMediaPlayer->PlayOnOpen = true;
+	InternalMediaPlayer->SetLooping(true);
+	
+	// Create MediaTexture
+	MediaTexture = NewObject<UMediaTexture>(GetTransientPackage(), NAME_None, RF_Transient);
 
-	// Disable fog on scene capture component of CGLayer
+	if (MediaTexture.IsValid())
 	{
-		TArray<USceneCaptureComponent2D*> CaptureComponents;
-		CGLayer->GetComponents(CaptureComponents);
-
-		for (USceneCaptureComponent2D* CaptureComponent : CaptureComponents)
-		{
-			// We need to disable Fog via the ShowFlagSettings property instead of directly in ShowFlags,
-			// because it will likely be overwritten since they are often replaced by the archetype defaults
-			// and then updated by the ShowFlagSettings.
-
-			FEngineShowFlagsSetting NewFlagSetting;
-			NewFlagSetting.ShowFlagName = FEngineShowFlags::FindNameByIndex(FEngineShowFlags::EShowFlag::SF_Fog);
-			NewFlagSetting.Enabled = false;
-
-			CaptureComponent->SetShowFlagSettings({ NewFlagSetting });
-		}
+		MediaTexture->AutoClear = true;
+		MediaTexture->SetMediaPlayer(InternalMediaPlayer.Get());
+		MediaTexture->UpdateResource();
 	}
 
-	// Create media player and media texture and assign it to the media input of the media plate layer.
-	//
-	for (UCompositingElementInput* Input : MediaPlate->GetInputsList())
+	// Play the media source, preferring time-synchronizable sources.
+	if (const UMediaProfile* MediaProfile = IMediaProfileManager::Get().GetCurrentMediaProfile())
 	{
-		UMediaTextureCompositingInput* MediaInput = Cast<UMediaTextureCompositingInput>(Input);
-
-		if (!MediaInput)
-		{
-			continue;
-		}
+		bool bFoundPreferredMediaSource = false;
 
-		// Create MediaPlayer
-
-		// Using a strong reference prevents the MediaPlayer from going stale when the level is saved.
-		MediaPlayer = TStrongObjectPtr<UMediaPlayer>(NewObject<UMediaPlayer>(
-			GetTransientPackage(),
-			MakeUniqueObjectName(GetTransientPackage(), UMediaPlayer::StaticClass())
-			));
-
-		if (!MediaPlayer.IsValid())
+		for (int32 MediaSourceIdx = 0; MediaSourceIdx < MediaProfile->NumMediaSources(); ++MediaSourceIdx)
 		{
-			UE_LOG(LogCameraCalibrationEditor, Warning, TEXT("Failed to create MediaPlayer"));
-			Cleanup();
-			return;
-		}
+			UMediaSource* MediaSource = MediaProfile->GetMediaSource(MediaSourceIdx);
 
-		MediaPlayer->PlayOnOpen = true;
-		MediaPlayer->SetLooping(true);
-		// Create MediaTexture
-
-		MediaTexture = NewObject<UMediaTexture>(GetTransientPackage(), NAME_None, RF_Transient);
-
-		if (MediaTexture.IsValid())
-		{
-			MediaTexture->AutoClear = true;
-			MediaTexture->SetMediaPlayer(MediaPlayer.Get());
-			MediaTexture->UpdateResource();
+			if (Cast<UTimeSynchronizableMediaSource>(MediaSource))
+			{
+				MediaProfileSourceIndex = MediaSourceIdx;
+				ExternalMediaTexture = MediaProfile->GetPlaybackManager()->OpenSourceFromIndex(MediaProfileSourceIndex, this);
+				bFoundPreferredMediaSource = true;
 
-			MediaInput->MediaSource = MediaTexture.Get();
+				// Break since we don't need to look for more MediaSources.
+				break;
+			}
 		}
 
-		// Play the media source, preferring time-synchronizable sources.
-		if (const UMediaProfile* MediaProfile = IMediaProfileManager::Get().GetCurrentMediaProfile())
+		if (!bFoundPreferredMediaSource)
 		{
-			bool bFoundPreferredMediaSource = false;
-
 			for (int32 MediaSourceIdx = 0; MediaSourceIdx < MediaProfile->NumMediaSources(); ++MediaSourceIdx)
 			{
-				UMediaSource* MediaSource = MediaProfile->GetMediaSource(MediaSourceIdx);
-
-				if (Cast<UTimeSynchronizableMediaSource>(MediaSource))
+				if (UMediaSource* MediaSource = MediaProfile->GetMediaSource(MediaSourceIdx))
 				{
-					MediaPlayer->OpenSource(MediaSource);
-					MediaPlayer->Play();
-
-					bFoundPreferredMediaSource = true;
+					MediaProfileSourceIndex = MediaSourceIdx;
+					ExternalMediaTexture = MediaProfile->GetPlaybackManager()->OpenSourceFromIndex(MediaProfileSourceIndex, this);
 
 					// Break since we don't need to look for more MediaSources.
 					break;
 				}
 			}
-
-			if (!bFoundPreferredMediaSource)
-			{
-				for (int32 MediaSourceIdx = 0; MediaSourceIdx < MediaProfile->NumMediaSources(); ++MediaSourceIdx)
-				{
-					if (UMediaSource* MediaSource = MediaProfile->GetMediaSource(MediaSourceIdx))
-					{
-						MediaPlayer->OpenSource(MediaSource);
-						MediaPlayer->Play();
-
-						// Break since we don't need to look for more MediaSources.
-						break;
-					}
-				}
-			}
 		}
-
-		// Break since don't need to look at more UMediaTextureCompositingInputs.
-		break;
 	}
 
-	// Create material pass that blends MediaPlate with CG.
-	//@todo Make this material selectable by the user.
-
-	MaterialPass = CastChecked<UCompositingElementMaterialPass>(
-		Comp->CreateNewTransformPass(TEXT("CGOverMedia"), UCompositingElementMaterialPass::StaticClass())
-		);
-
-	if (!MaterialPass.IsValid())
-	{
-		UE_LOG(LogCameraCalibrationEditor, Warning, TEXT("Failed to create 'CGOverMedia' UCompositingElementMaterialPass"));
-		Cleanup();
-		return;
-	}
-
-	// Create Material
-	const FString MaterialPath = TEXT("/CameraCalibration/Materials/M_SimulcamCalib.M_SimulcamCalib");
-	UMaterial* Material = Cast<UMaterial>(FSoftObjectPath(MaterialPath).TryLoad());
-
-	if (!Material)
-	{
-		UE_LOG(LogCameraCalibrationEditor, Warning, TEXT("Failed to load %s"), *MaterialPath);
-		Cleanup();
-		return;
-	}
-
-	MaterialPass->SetMaterialInterface(Material);
-
-	// ActorLabel should coincide with ElementName. Name may not.
-	MaterialPass->SetParameterMapping(TEXT("CG"), *CGLayer->GetActorLabel());
-	MaterialPass->SetParameterMapping(TEXT("MediaPlate"), *MediaPlate->GetActorLabel());
-
-	// Create new overlay transform passes
-	CreateOverlayPass(TEXT("Calibration Step Overlay"), ToolOverlayPass, ToolOverlayRenderTarget);
-	CreateOverlayPass(TEXT("User Selected Overlay"), UserOverlayPass, UserOverlayRenderTarget);
-
-	const FIntPoint InitialCompRenderResolution = Comp->GetRenderResolution();
-
+	const FIntPoint DefaultResolution(1920, 1080);
 	if (ULensFile* LensFilePtr = LensFile.Get())
 	{
 		// Initialize the simulcam info of the LensFile to match the initial render resolution of the comp
-		const float CompAspectRatio = (InitialCompRenderResolution.Y != 0) ? InitialCompRenderResolution.X / (float)InitialCompRenderResolution.Y : 1.0f;
-		LensFilePtr->SimulcamInfo.CGLayerAspectRatio = CompAspectRatio;
-		LensFilePtr->SimulcamInfo.MediaPlateAspectRatio = CompAspectRatio;
+		const float AspectRatio = (DefaultResolution.Y != 0) ? DefaultResolution.X / (float)DefaultResolution.Y : 1.0f;
+		LensFilePtr->SimulcamInfo.CGLayerAspectRatio = AspectRatio;
+		LensFilePtr->SimulcamInfo.MediaPlateAspectRatio = AspectRatio;
 
 		// If the Camera Feed is not specifically overriden by the user, initialize it to the initial dimensions of the comp as well
 		// If the camera feed exactly matches the comp's resolution, then no aspect ratio correction is needed
 		if (!LensFilePtr->CameraFeedInfo.IsOverridden())
 		{
-			LensFilePtr->CameraFeedInfo.SetDimensions(InitialCompRenderResolution);
+			LensFilePtr->CameraFeedInfo.SetDimensions(DefaultResolution);
 		}
-		else
-		{
-			// If the Camera Feed has been overridden by the user, update the aspect ratio correction material parameters
-			UpdateAspectRatioCorrection();
-		}
-	}
-
-	URenderTargetCompositingOutput* RTOutput = Cast<URenderTargetCompositingOutput>(Comp->CreateNewOutputPass(
-		TEXT("SimulcamCalOutput"),
-		URenderTargetCompositingOutput::StaticClass())
-		);
-
-	if (!RTOutput)
-	{
-		UE_LOG(LogCameraCalibrationEditor, Warning, TEXT("Failed to create URenderTargetCompositingOutput"));
-		Cleanup();
-		return;
-	}
-
-	// Create RenderTarget
-	RenderTarget = NewObject<UTextureRenderTarget2D>(
-		GetTransientPackage(),
-		MakeUniqueObjectName(GetTransientPackage(), UTextureRenderTarget2D::StaticClass())
-		);
-
-	if (!RenderTarget.IsValid())
-	{
-		UE_LOG(LogCameraCalibrationEditor, Warning, TEXT("Failed to create UTextureRenderTarget2D"));
-		Cleanup();
-		return;
 	}
 
-	RenderTarget->RenderTargetFormat = RTF_RGBA16f;
-	RenderTarget->ClearColor = FLinearColor::Black;
-	RenderTarget->bAutoGenerateMips = false;
-	RenderTarget->InitAutoFormat(InitialCompRenderResolution.X, InitialCompRenderResolution.Y);
-	RenderTarget->UpdateResourceImmediate(true);
-
-	// Assign the RT to the compositing output
-	RTOutput->RenderTarget = RenderTarget.Get();
-
-	// By default, use a camera with our lens. If not found, let composure pick one.
+	// By default, use a camera with our lens. If not found, use the first camera found in the level
 	if (ACameraActor* CameraGuess = FindFirstCameraWithCurrentLens())
 	{
 		SetCamera(CameraGuess);
 	}
 	else
 	{
-		SetCamera(Comp->FindTargetCamera());
-	}
-}
-
-UCompositingElementMaterialPass* FCameraCalibrationStepsController::GetOverlayMaterialPass(EOverlayPassType OverlayPassType) const
-{
-	if (OverlayPassType == EOverlayPassType::ToolOverlay)
-	{
-		return ToolOverlayPass.Get();
-	}
-	else if (OverlayPassType == EOverlayPassType::UserOverlay)
-	{
-		return UserOverlayPass.Get();
-	}
-
-	return nullptr;
-}
-
-UTextureRenderTarget2D* FCameraCalibrationStepsController::GetOverlayRenderTarget(EOverlayPassType OverlayPassType) const
-{
-	if (OverlayPassType == EOverlayPassType::ToolOverlay)
-	{
-		return ToolOverlayRenderTarget.Get();
-	}
-	else if (OverlayPassType == EOverlayPassType::UserOverlay)
-	{
-		return UserOverlayRenderTarget.Get();
+		SetCamera(FindDefaultCamera());
 	}
 
-	return nullptr;
+	// Create a render target that can be used to extract media texture pixels
+	CachedMediaRenderTarget = TStrongObjectPtr<UTextureRenderTarget2D>(NewObject<UTextureRenderTarget2D>(
+		GetTransientPackage(),
+		MakeUniqueObjectName(GetTransientPackage(), UTextureRenderTarget2D::StaticClass())));
+	
+	CachedMediaRenderTarget->RenderTargetFormat = ETextureRenderTargetFormat::RTF_RGBA8;
+	CachedMediaRenderTarget->ClearColor = FLinearColor::Black;
+	CachedMediaRenderTarget->bAutoGenerateMips = false;
+	CachedMediaRenderTarget->InitAutoFormat(DefaultResolution.X, DefaultResolution.Y);
+	CachedMediaRenderTarget->UpdateResourceImmediate(true);
 }
 
 UMaterialInterface* FCameraCalibrationStepsController::GetOverlayMaterial(EOverlayPassType OverlayPassType) const
@@ -950,19 +521,29 @@ UMaterialInterface* FCameraCalibrationStepsController::GetOverlayMaterial(EOverl
 
 bool FCameraCalibrationStepsController::IsOverlayEnabled(EOverlayPassType OverlayPassType) const
 {
-	if (UCompositingElementMaterialPass* OverlayPass = GetOverlayMaterialPass(OverlayPassType))
+	switch (OverlayPassType)
 	{
-		return OverlayPass->IsPassEnabled();
-	}
+	case EOverlayPassType::ToolOverlay:
+		return bToolOverlayEnabled;
 
+	case EOverlayPassType::UserOverlay:
+		return bUserOverlayEnabled;
+	}
+	
 	return false;
 }
 
 void FCameraCalibrationStepsController::SetOverlayEnabled(const bool bEnabled, EOverlayPassType OverlayPassType)
 {
-	if (UCompositingElementMaterialPass* OverlayPass = GetOverlayMaterialPass(OverlayPassType))
+	switch (OverlayPassType)
 	{
-		OverlayPass->SetPassEnabled(bEnabled);
+	case EOverlayPassType::ToolOverlay:
+		bToolOverlayEnabled = bEnabled;
+		break;
+
+	case EOverlayPassType::UserOverlay:
+		bUserOverlayEnabled = bEnabled;
+		break;
 	}
 }
 
@@ -980,7 +561,6 @@ void FCameraCalibrationStepsController::SetOverlayMaterial(UMaterialInterface* I
 	// If the overlay is non-null, refresh the overlay and enable/disable the overlay as necessary
 	if (InOverlay)
 	{
-		RefreshOverlay(OverlayPassType);
 		SetOverlayEnabled(bShowOverlay, OverlayPassType);
 	}
 	else
@@ -990,129 +570,14 @@ void FCameraCalibrationStepsController::SetOverlayMaterial(UMaterialInterface* I
 	}
 }
 
-void FCameraCalibrationStepsController::RefreshOverlay(EOverlayPassType OverlayPassType)
-{
-	if (UCompositingElementMaterialPass* OverlayPass = GetOverlayMaterialPass(OverlayPassType))
-	{
-		if (UTextureRenderTarget2D* OverlayRenderTarget = GetOverlayRenderTarget(OverlayPassType))
-		{
-			if (UMaterialInterface* OverlayMaterial = GetOverlayMaterial(OverlayPassType))
-			{
-				// Clear the overlay render target and draw to it again using the current overlay material
-				UKismetRenderingLibrary::ClearRenderTarget2D(OverlayPass, OverlayRenderTarget, FLinearColor::Transparent);
-				UKismetRenderingLibrary::DrawMaterialToRenderTarget(OverlayPass, OverlayRenderTarget, OverlayMaterial);
-			}
-		}
-	}
-}
-
-void FCameraCalibrationStepsController::CreateMediaPlateOutput()
-{
-	if (!MediaPlate.IsValid())
-	{
-		return;
-	}
-
-	URenderTargetCompositingOutput* RTOutput = Cast<URenderTargetCompositingOutput>(MediaPlate->CreateNewOutputPass(
-		TEXT("MediaPlateOutput"),
-		URenderTargetCompositingOutput::StaticClass())
-		);
-
-	if (!RTOutput)
-	{
-		UE_LOG(LogCameraCalibrationEditor, Warning, TEXT("Failed to create URenderTargetCompositingOutput for the MediaPlateOutput"));
-		Cleanup();
-		return;
-	}
-
-	// Create RenderTarget
-	MediaPlateRenderTarget = NewObject<UTextureRenderTarget2D>(
-		GetTransientPackage(),
-		MakeUniqueObjectName(GetTransientPackage(), UTextureRenderTarget2D::StaticClass())
-		);
-
-	if (!MediaPlateRenderTarget.IsValid())
-	{
-		UE_LOG(LogCameraCalibrationEditor, Warning, TEXT("Failed to create UTextureRenderTarget2D for the MediaPlateOutput"));
-		Cleanup();
-		return;
-	}
-
-	const FIntPoint InitialCompRenderResolution = GetCompRenderResolution();
-
-	MediaPlateRenderTarget->RenderTargetFormat = ETextureRenderTargetFormat::RTF_RGBA8;
-	MediaPlateRenderTarget->ClearColor = FLinearColor::Black;
-	MediaPlateRenderTarget->bAutoGenerateMips = false;
-	MediaPlateRenderTarget->InitAutoFormat(InitialCompRenderResolution.X, InitialCompRenderResolution.Y);
-	MediaPlateRenderTarget->UpdateResourceImmediate(true);
-
-	// Assign the RT to the compositing output
-	RTOutput->RenderTarget = MediaPlateRenderTarget.Get();
-}
-
-void FCameraCalibrationStepsController::CreateOverlayPass(FName PassName, TWeakObjectPtr<UCompositingElementMaterialPass>& OverlayPass, TWeakObjectPtr<UTextureRenderTarget2D>& OverlayRenderTarget)
-{
-	// Create new overlay transform pass
-	OverlayPass = CastChecked<UCompositingElementMaterialPass>(
-		Comp->CreateNewTransformPass(PassName, UCompositingElementMaterialPass::StaticClass())
-		);
-
-	if (!OverlayPass.IsValid())
-	{
-		UE_LOG(LogCameraCalibrationEditor, Warning, TEXT("Failed to create 'Overlay' UCompositingElementMaterialPass"));
-		Cleanup();
-		return;
-	}
-
-	if (UMaterialInterface* OverlayBaseMaterial = Cast<UMaterialInterface>(StaticLoadObject(UMaterialInterface::StaticClass(), NULL, TEXT("/CameraCalibration/Materials/M_OverlayBase.M_OverlayBase"))))
-	{
-		OverlayPass->SetMaterialInterface(OverlayBaseMaterial);
-	}
-
-	OverlayPass->SetPassEnabled(false);
-
-	OverlayRenderTarget = NewObject<UTextureRenderTarget2D>(
-		GetTransientPackage(),
-		MakeUniqueObjectName(GetTransientPackage(), UTextureRenderTarget2D::StaticClass())
-		);
-
-	if (!OverlayRenderTarget.IsValid())
-	{
-		UE_LOG(LogCameraCalibrationEditor, Warning, TEXT("Failed to create UTextureRenderTarget2D for the OverlayPass"));
-		Cleanup();
-		return;
-	}
-
-	const FIntPoint InitialCompRenderResolution = GetCompRenderResolution();
-
-	OverlayRenderTarget->RenderTargetFormat = ETextureRenderTargetFormat::RTF_RGBA8;
-	OverlayRenderTarget->ClearColor = FLinearColor::Transparent;
-	OverlayRenderTarget->bAutoGenerateMips = false;
-	OverlayRenderTarget->InitAutoFormat(InitialCompRenderResolution.X, InitialCompRenderResolution.Y);
-	OverlayRenderTarget->UpdateResourceImmediate(true);
-
-	OverlayPass->Material.SetTextureOverride(FName(TEXT("OverlayTexture")), OverlayRenderTarget.Get());
-}
-
 float FCameraCalibrationStepsController::GetWiperWeight() const
 {
-	float Weight = 0.5f;
-
-	if (!MaterialPass.IsValid())
-	{
-		return Weight;
-	}
-
-	MaterialPass->Material.GetScalarOverride(TEXT("WiperWeight"), Weight);
-	return Weight;
+	return WiperWeight;
 }
 
 void FCameraCalibrationStepsController::SetWiperWeight(float InWeight)
 {
-	if (MaterialPass.IsValid())
-	{
-		MaterialPass->Material.SetScalarOverride(TEXT("WiperWeight"), InWeight);
-	}
+	WiperWeight = InWeight;
 }
 
 void FCameraCalibrationStepsController::SetCamera(ACameraActor* InCamera)
@@ -1131,48 +596,28 @@ void FCameraCalibrationStepsController::SetCamera(ACameraActor* InCamera)
 	{
 		CineCameraComponent = CameraComponents[0];
 	}
+}
 
-	// Update the Comp with this camera
-	if (Comp.IsValid())
-	{
-		Comp->SetTargetCamera(InCamera);
-		EnableDistortionInCG();
-	}
+ACameraActor* FCameraCalibrationStepsController::GetCamera() const
+{
+	return Camera.Get();
 }
 
-void FCameraCalibrationStepsController::EnableDistortionInCG()
+float FCameraCalibrationStepsController::GetCameraFeedAspectRatio() const
 {
-	if (!Comp.IsValid())
+	TStrongObjectPtr<ULensFile> PinnedLensFile = LensFile.Pin();
+	if (!PinnedLensFile)
 	{
-		return;
+		return 1.0f;
 	}
 
-	for (ACompositingElement* Element : Comp->GetChildElements())
+	if (PinnedLensFile->CameraFeedInfo.IsValid())
 	{
-		ACompositingCaptureBase* CaptureBase = Cast<ACompositingCaptureBase>(Element);
-
-		if (!CaptureBase)
-		{
-			continue;
-		}
-
-		// Set the LensComponent used by the CG layer to match the LensComponent in use by our target camera
-		if (ACameraActor* const TargetCamera = Camera.Get())
-		{
-			if (ULensComponent* const TargetLens = FindLensComponentOnCamera(TargetCamera))
-			{
-				CaptureBase->SetLens(TargetLens);
-			}
-		}
-
-		// Enable distortion on the CG compositing layer
-		CaptureBase->SetApplyDistortion(true);
+		const FVector2D CameraFeedDimensions(PinnedLensFile->CameraFeedInfo.GetDimensions());
+		return CameraFeedDimensions.X / CameraFeedDimensions.Y;
 	}
-}
 
-ACameraActor* FCameraCalibrationStepsController::GetCamera() const
-{
-	return Camera.Get();
+	return 1.0f;
 }
 
 void FCameraCalibrationStepsController::OnSimulcamViewportClicked(const FGeometry& MyGeometry, const FPointerEvent& MouseEvent)
@@ -1233,7 +678,7 @@ void FCameraCalibrationStepsController::OnSimulcamViewportMarqueeSelect(FVector2
 FReply FCameraCalibrationStepsController::OnRewindButtonClicked()
 {
 	// Rewind to the beginning of the media
-	if (MediaPlayer.IsValid())
+	if (UMediaPlayer* MediaPlayer = GetMediaPlayer())
 	{
 		MediaPlayer->Rewind();
 	}
@@ -1244,7 +689,7 @@ FReply FCameraCalibrationStepsController::OnRewindButtonClicked()
 FReply FCameraCalibrationStepsController::OnReverseButtonClicked()
 {
 	// Increase the reverse media playback rate
-	if (MediaPlayer.IsValid())
+	if (UMediaPlayer* MediaPlayer = GetMediaPlayer())
 	{
 		MediaPlayer->SetRate(GetFasterReverseRate());
 	}
@@ -1254,7 +699,7 @@ FReply FCameraCalibrationStepsController::OnReverseButtonClicked()
 
 FReply FCameraCalibrationStepsController::OnStepBackButtonClicked() 
 {
-	if (MediaPlayer.IsValid())
+	if (UMediaPlayer* MediaPlayer = GetMediaPlayer())
 	{
 		const float DefaultStepRateInMilliseconds = GetDefault<UCameraCalibrationEditorSettings>()->DefaultMediaStepRateInMilliseconds;
 		const bool bForceDefaultStepRate = GetDefault<UCameraCalibrationEditorSettings>()->bForceDefaultMediaStepRate;
@@ -1290,7 +735,7 @@ FReply FCameraCalibrationStepsController::OnStepBackButtonClicked()
 
 FReply FCameraCalibrationStepsController::OnPlayButtonClicked() 
 {
-	if (MediaPlayer.IsValid())
+	if (UMediaPlayer* MediaPlayer = GetMediaPlayer())
 	{
 		MediaPlayer->Play();
 	}
@@ -1300,7 +745,7 @@ FReply FCameraCalibrationStepsController::OnPlayButtonClicked()
 
 FReply FCameraCalibrationStepsController::OnPauseButtonClicked() 
 {
-	if (MediaPlayer.IsValid())
+	if (UMediaPlayer* MediaPlayer = GetMediaPlayer())
 	{
 		MediaPlayer->Pause();
 	}
@@ -1310,7 +755,7 @@ FReply FCameraCalibrationStepsController::OnPauseButtonClicked()
 
 FReply FCameraCalibrationStepsController::OnStepForwardButtonClicked()
 {
-	if (MediaPlayer.IsValid())
+	if (UMediaPlayer* MediaPlayer = GetMediaPlayer())
 	{
 		const float DefaultStepRateInMilliseconds = GetDefault<UCameraCalibrationEditorSettings>()->DefaultMediaStepRateInMilliseconds;
 		const bool bForceDefaultStepRate = GetDefault<UCameraCalibrationEditorSettings>()->bForceDefaultMediaStepRate;
@@ -1349,7 +794,7 @@ FReply FCameraCalibrationStepsController::OnStepForwardButtonClicked()
 FReply FCameraCalibrationStepsController::OnForwardButtonClicked() 
 {
 	// Increase the forward media playback rate
-	if (MediaPlayer.IsValid())
+	if (UMediaPlayer* MediaPlayer = GetMediaPlayer())
 	{
 		MediaPlayer->SetRate(GetFasterForwardRate());
 	}
@@ -1359,7 +804,7 @@ FReply FCameraCalibrationStepsController::OnForwardButtonClicked()
 
 bool FCameraCalibrationStepsController::DoesMediaSupportSeeking() const
 {
-	if (MediaPlayer.IsValid())
+	if (UMediaPlayer* MediaPlayer = GetMediaPlayer())
 	{
 		return MediaPlayer->SupportsSeeking();
 	}
@@ -1369,7 +814,7 @@ bool FCameraCalibrationStepsController::DoesMediaSupportSeeking() const
 
 bool FCameraCalibrationStepsController::DoesMediaSupportNextReverseRate() const
 {
-	if (MediaPlayer.IsValid())
+	if (UMediaPlayer* MediaPlayer = GetMediaPlayer())
 	{
 		constexpr bool Unthinned = false;
 		return MediaPlayer->SupportsRate(GetFasterReverseRate(), Unthinned);
@@ -1380,7 +825,7 @@ bool FCameraCalibrationStepsController::DoesMediaSupportNextReverseRate() const
 
 bool FCameraCalibrationStepsController::DoesMediaSupportNextForwardRate() const
 {
-	if (MediaPlayer.IsValid())
+	if (UMediaPlayer* MediaPlayer = GetMediaPlayer())
 	{
 		constexpr bool Unthinned = false;
 		return MediaPlayer->SupportsRate(GetFasterForwardRate(), Unthinned);
@@ -1391,7 +836,7 @@ bool FCameraCalibrationStepsController::DoesMediaSupportNextForwardRate() const
 
 float FCameraCalibrationStepsController::GetFasterReverseRate() const
 {
-	if (MediaPlayer.IsValid())
+	if (UMediaPlayer* MediaPlayer = GetMediaPlayer())
 	{
 		// Get the current playback rate of the media player
 		float Rate = MediaPlayer->GetRate();
@@ -1411,7 +856,7 @@ float FCameraCalibrationStepsController::GetFasterReverseRate() const
 
 float FCameraCalibrationStepsController::GetFasterForwardRate() const
 {
-	if (MediaPlayer.IsValid())
+	if (UMediaPlayer* MediaPlayer = GetMediaPlayer())
 	{
 		// Get the current playback rate of the media player
 		float Rate = MediaPlayer->GetRate();
@@ -1465,49 +910,66 @@ ACameraActor* FCameraCalibrationStepsController::FindFirstCameraWithCurrentLens(
 	return FirstCamera;
 }
 
-void FCameraCalibrationStepsController::TogglePlay()
+ACameraActor* FCameraCalibrationStepsController::FindDefaultCamera() const
 {
-	if (!MediaPlayer.IsValid())
+	for (TActorIterator<ACineCameraActor> CamActorIt(GetWorld()); CamActorIt; ++CamActorIt)
 	{
-		return;
+		ACineCameraActor* FoundActor = *CamActorIt;
+		if (IsValid(FoundActor))
+		{
+			return FoundActor;
+		}
 	}
 
-	//@todo Eventually pause should cache the texture instead of relying on player play/pause support.
-
-	if (IsPaused())
+	for (TActorIterator<ACameraActor> CamActorIt(GetWorld()); CamActorIt; ++CamActorIt)
 	{
-		MediaPlayer->Play();
+		ACameraActor* FoundActor = *CamActorIt;
+		if (IsValid(FoundActor))
+		{
+			return FoundActor;
+		}
 	}
-	else
+
+	return nullptr;
+}
+
+void FCameraCalibrationStepsController::TogglePlay()
+{
+	//@todo Eventually pause should cache the texture instead of relying on player play/pause support.
+
+	if (UMediaPlayer* MediaPlayer = GetMediaPlayer())
 	{
-		// TODO: Trigger the current step (and ultimately the algo) to cache any data it cares about (like 3D scene data)
-		MediaPlayer->Pause();
+		if (MediaPlayer->IsPaused())
+		{
+			MediaPlayer->Play();
+		}
+		else
+		{
+			// TODO: Trigger the current step (and ultimately the algo) to cache any data it cares about (like 3D scene data)
+			MediaPlayer->Pause();
+		}
 	}
 }
 
 void FCameraCalibrationStepsController::Play()
 {
-	if (!MediaPlayer.IsValid())
+	if (UMediaPlayer* MediaPlayer = GetMediaPlayer())
 	{
-		return;
+		MediaPlayer->Play();
 	}
-
-	MediaPlayer->Play();
 }
 
 void FCameraCalibrationStepsController::Pause()
 {
-	if (!MediaPlayer.IsValid())
+	if (UMediaPlayer* MediaPlayer = GetMediaPlayer())
 	{
-		return;
+		MediaPlayer->Pause();
 	}
-
-	MediaPlayer->Pause();
 }
 
 bool FCameraCalibrationStepsController::IsPaused() const
 {
-	if (MediaPlayer.IsValid())
+	if (UMediaPlayer* MediaPlayer = GetMediaPlayer())
 	{
 		return MediaPlayer->IsPaused();
 	}
@@ -1569,7 +1031,7 @@ bool FCameraCalibrationStepsController::SetMediaSource(UMediaSource* InMediaSour
 {
 	ClearMedia();
 	
-	if (!MediaPlayer.IsValid())
+	if (!InternalMediaPlayer.IsValid())
 	{
 		return false;
 	}
@@ -1579,8 +1041,8 @@ bool FCameraCalibrationStepsController::SetMediaSource(UMediaSource* InMediaSour
 		return true;
 	}
 	
-	MediaPlayer->OpenSource(InMediaSource);
-	MediaPlayer->Play();
+	InternalMediaPlayer->OpenSource(InMediaSource);
+	InternalMediaPlayer->Play();
 	return false;
 }
 
@@ -1588,67 +1050,45 @@ bool FCameraCalibrationStepsController::SetMediaTexture(UMediaTexture* InMediaTe
 {
 	ClearMedia();
 	
-	if (!MediaPlate.IsValid())
-	{
-		return false;
-	}
-	
 	ExternalMediaTexture = InMediaTexture;
+	return true;
+}
 
-	UMediaTexture* Texture = nullptr;
-	if (ExternalMediaTexture.IsValid())
-	{
-		Texture = ExternalMediaTexture.Get();
-	}
-	else
-	{
-		Texture = MediaTexture.Get();
-	}
-
-	if (Texture)
+bool FCameraCalibrationStepsController::SetMediaProfileSource(UMediaSource* InMediaProfileSource)
+{
+	ClearMedia();
+	
+	if (UMediaProfile* MediaProfile = IMediaProfileManager::Get().GetCurrentMediaProfile())
 	{
-		for (UCompositingElementInput* Input : MediaPlate->GetInputsList())
+		MediaProfileSourceIndex = MediaProfile->FindMediaSourceIndex(InMediaProfileSource);
+		if (MediaProfileSourceIndex == INDEX_NONE)
 		{
-			UMediaTextureCompositingInput* MediaInput = Cast<UMediaTextureCompositingInput>(Input);
-
-			if (!MediaInput)
-			{
-				continue;
-			}
-
-			MediaInput->MediaSource = Texture;
-			break;
+			return false;
 		}
+
+		ExternalMediaTexture = MediaProfile->GetPlaybackManager()->OpenSourceFromIndex(MediaProfileSourceIndex, this);
+		return ExternalMediaTexture.IsValid();
 	}
-	
-	return true;
+
+	return false;
 }
 
 void FCameraCalibrationStepsController::ClearMedia()
 {
-	if (!MediaPlate.IsValid())
+	if (MediaProfileSourceIndex != INDEX_NONE)
 	{
-		return;
-	}
-
-	MediaPlayer->Close();
-
-	ExternalMediaTexture = nullptr;
-	if (MediaTexture.IsValid())
-	{
-		for (UCompositingElementInput* Input : MediaPlate->GetInputsList())
+		if (UMediaProfile* MediaProfile = IMediaProfileManager::Get().GetCurrentMediaProfile())
 		{
-			UMediaTextureCompositingInput* MediaInput = Cast<UMediaTextureCompositingInput>(Input);
-
-			if (!MediaInput)
-			{
-				continue;
-			}
-
-			MediaInput->MediaSource = MediaTexture.Get();
-			break;
+			UMediaProfilePlaybackManager::FCloseSourceArgs Args;
+			Args.Consumer = this;
+			
+			MediaProfile->GetPlaybackManager()->CloseSourceFromIndex(MediaProfileSourceIndex, Args);
 		}
 	}
+	
+	MediaProfileSourceIndex = INDEX_NONE;
+	InternalMediaPlayer->Close();
+	ExternalMediaTexture = nullptr;
 }
 
 void FCameraCalibrationStepsController::GetMediaProfileSources(TArray<TWeakObjectPtr<UMediaSource>>& OutSources) const
@@ -1685,22 +1125,35 @@ TArray<TSoftObjectPtr<UMediaTexture>> FCameraCalibrationStepsController::GetMedi
 /** Gets the current media source url being played. Empty if None */
 FString FCameraCalibrationStepsController::GetMediaSourceUrl() const
 {
-	if (!MediaPlayer.IsValid())
+	if (UMediaPlayer* MediaPlayer = GetMediaPlayer())
 	{
-		return TEXT("");
+		return MediaPlayer->GetUrl();
 	}
 
-	return MediaPlayer->GetUrl();
+	return TEXT("");
 }
 
 UMediaSource* FCameraCalibrationStepsController::GetMediaSource() const
 {
-	if (!MediaPlayer.IsValid())
+	if (!InternalMediaPlayer.IsValid())
 	{
 		return nullptr;
 	}
 
-	return MediaPlayer->GetPlaylist()->Get(MediaPlayer->GetPlaylistIndex());
+	return InternalMediaPlayer->GetPlaylist()->Get(InternalMediaPlayer->GetPlaylistIndex());
+}
+
+UMediaSource* FCameraCalibrationStepsController::GetMediaProfileSource() const
+{
+	if (MediaProfileSourceIndex != INDEX_NONE)
+	{
+		if (UMediaProfile* MediaProfile = IMediaProfileManager::Get().GetCurrentMediaProfile())
+		{
+			return MediaProfile->GetMediaSource(MediaProfileSourceIndex);
+		}
+	}
+
+	return nullptr;
 }
 
 UMediaTexture* FCameraCalibrationStepsController::GetMediaTexture() const
@@ -1708,6 +1161,21 @@ UMediaTexture* FCameraCalibrationStepsController::GetMediaTexture() const
 	return ExternalMediaTexture.Get();
 }
 
+UMediaPlayer* FCameraCalibrationStepsController::GetMediaPlayer() const
+{
+	if (ExternalMediaTexture.IsValid())
+	{
+		return ExternalMediaTexture->GetMediaPlayer();
+	}
+
+	if (InternalMediaPlayer.IsValid())
+	{
+		return InternalMediaPlayer.Get();
+	}
+
+	return nullptr;
+}
+
 const TConstArrayView<TStrongObjectPtr<UCameraCalibrationStep>> FCameraCalibrationStepsController::GetCalibrationSteps() const
 {
 	return TConstArrayView<TStrongObjectPtr<UCameraCalibrationStep>>(CalibrationSteps);
@@ -1738,14 +1206,40 @@ void FCameraCalibrationStepsController::Initialize()
 {
 	// Not doing these in the constructor so that SharedThis can be used.
 
-	CreateComp();
-	CreateMediaPlateOutput();
+	InitializeMediaPlayer();
 	CreateSteps();
+	
+	IMediaProfileManager::Get().OnMediaProfileChanged().AddSP(this, &FCameraCalibrationStepsController::OnActiveMediaProfileChanged);
 }
 
-UTextureRenderTarget2D* FCameraCalibrationStepsController::GetMediaPlateRenderTarget() const
+UMediaTexture* FCameraCalibrationStepsController::GetMediaOverlayTexture() const
 {
-	return MediaPlateRenderTarget.Get();
+	if (ExternalMediaTexture.IsValid())
+	{
+		return ExternalMediaTexture.Get();
+	}
+
+	if (InternalMediaPlayer->IsReady())
+	{
+		return MediaTexture.IsValid() ? MediaTexture.Get() : nullptr;
+	}
+
+	return nullptr;
+}
+
+FIntPoint FCameraCalibrationStepsController::GetMediaOverlayResolution() const
+{
+	if (ExternalMediaTexture.IsValid())
+	{
+		return FIntPoint(ExternalMediaTexture.Get()->GetSurfaceWidth(), ExternalMediaTexture.Get()->GetSurfaceHeight());
+	}
+
+	if (InternalMediaPlayer->IsReady())
+	{
+		return MediaTexture.IsValid() ? FIntPoint(MediaTexture.Get()->GetSurfaceWidth(), MediaTexture.Get()->GetSurfaceHeight()) : FIntPoint::ZeroValue;
+	}
+
+	return FIntPoint::ZeroValue;
 }
 
 bool FCameraCalibrationStepsController::CalculateNormalizedMouseClickPosition(const FGeometry& MyGeometry, const FPointerEvent& MouseEvent, FVector2f& OutPosition, ESimulcamViewportPortion ViewportPortion) const
@@ -1777,10 +1271,10 @@ bool FCameraCalibrationStepsController::CalculateNormalizedMouseClickPosition(co
 		}
 
 		const FIntPoint CameraFeedDimensions = LensFilePtr->CameraFeedInfo.GetDimensions();
-		const FIntPoint CompRenderResolution = GetCompRenderResolution();
+		const FIntPoint MediaOverlayResolution = GetMediaOverlayResolution();
 
-		const float AspectRatioCorrectionX = CompRenderResolution.X / (float)CameraFeedDimensions.X;
-		const float AspectRatioCorrectionY = CompRenderResolution.Y / (float)CameraFeedDimensions.Y;
+		const float AspectRatioCorrectionX = MediaOverlayResolution.X / (float)CameraFeedDimensions.X;
+		const float AspectRatioCorrectionY = MediaOverlayResolution.Y / (float)CameraFeedDimensions.Y;
 
 		XNormalized = ((XNormalized - 0.5f) * AspectRatioCorrectionX) + 0.5f;
 		YNormalized = ((YNormalized - 0.5f) * AspectRatioCorrectionY) + 0.5f;
@@ -1800,26 +1294,39 @@ bool FCameraCalibrationStepsController::CalculateNormalizedMouseClickPosition(co
 
 bool FCameraCalibrationStepsController::ReadMediaPixels(TArray<FColor>& Pixels, FIntPoint& Size, FText& OutErrorMessage, ESimulcamViewportPortion ViewportPortion) const
 {
-	// Get the media plate texture render target 2d
-
-	if (!MediaPlateRenderTarget.IsValid())
+	// Renders a media texture to a render target (via canvas rendering) and reads those textures to the output pixel array
+	// NOTE: UMediaTexture's texture resource is both an FTextureResource and an FRenderTarget, so in theory we could read the media texture directly,
+	// but such functionality is not exposed, nor is it likely the pixels are currently configured for CPU read access. As such, we must output to an
+	// intermediate render target to access the pixels
+	UMediaTexture* MediaOverlayTexture = GetMediaOverlayTexture();
+	if (!MediaOverlayTexture)
 	{
-		OutErrorMessage = LOCTEXT("InvalidMediaPlateRenderTarget", "Invalid MediaPlateRenderTarget");
+		OutErrorMessage = LOCTEXT("InvalidMediaTexture", "Invalid MediaTexture");
 		return false;
 	}
 
-	// Extract its render target resource
-	FRenderTarget* MediaRenderTarget = MediaPlateRenderTarget->GameThread_GetRenderTargetResource();
-
-	if (!MediaRenderTarget)
+	const FIntPoint MediaTextureSize(MediaOverlayTexture->GetSurfaceWidth(), MediaOverlayTexture->GetSurfaceHeight());
+	if (CachedMediaRenderTarget->SizeX != MediaTextureSize.X || CachedMediaRenderTarget->SizeY != MediaTextureSize.Y)
 	{
-		OutErrorMessage = LOCTEXT("InvalidRenderTargetResource", "MediaPlateRenderTarget did not have a RenderTarget resource");
-		return false;
+		CachedMediaRenderTarget->ResizeTarget(MediaTextureSize.X, MediaTextureSize.Y);
+		FlushRenderingCommands();
 	}
+	
+	UKismetRenderingLibrary::ClearRenderTarget2D(GetWorld(), CachedMediaRenderTarget.Get(), FLinearColor::Black);
 
-	if (MediaPlateRenderTarget->RenderTargetFormat != ETextureRenderTargetFormat::RTF_RGBA8)
+	UCanvas* Canvas;
+	FVector2D CanvasToRenderTargetSize;
+	FDrawToRenderTargetContext RenderTargetContext;
+	UKismetRenderingLibrary::BeginDrawCanvasToRenderTarget(GetWorld(), CachedMediaRenderTarget.Get(), /*out*/ Canvas, /*out*/ CanvasToRenderTargetSize, /*out*/ RenderTargetContext);
+	{
+		Canvas->K2_DrawTexture(MediaOverlayTexture, FVector2D::ZeroVector, MediaTextureSize, FVector2D::ZeroVector);
+	}
+	UKismetRenderingLibrary::EndDrawCanvasToRenderTarget(GetWorld(), RenderTargetContext);
+	
+	FTextureRenderTargetResource* MediaRenderTarget = CachedMediaRenderTarget->GameThread_GetRenderTargetResource();
+	if (!MediaRenderTarget)
 	{
-		OutErrorMessage = LOCTEXT("InvalidFormat", "MediaPlateRenderTarget did not have the expected RTF_RGBA8 format");
+		OutErrorMessage = LOCTEXT("InvalidRenderTargetResource", "Media render target did not have a RenderTarget resource");
 		return false;
 	}
 
@@ -1873,4 +1380,21 @@ bool FCameraCalibrationStepsController::ReadMediaPixels(TArray<FColor>& Pixels,
 	return true;
 }
 
+void FCameraCalibrationStepsController::OnActiveMediaProfileChanged(UMediaProfile* PreviousMediaProfile, UMediaProfile* NewMediaProfile)
+{
+	if (MediaProfileSourceIndex != INDEX_NONE)
+	{
+		if (PreviousMediaProfile)
+		{
+			UMediaProfilePlaybackManager::FCloseSourceArgs Args;
+			Args.Consumer = this;
+			
+			PreviousMediaProfile->GetPlaybackManager()->CloseSourceFromIndex(MediaProfileSourceIndex, Args);
+		}
+		MediaProfileSourceIndex = INDEX_NONE;
+	}
+
+	ClearMedia();
+}
+
 #undef LOCTEXT_NAMESPACE
diff --git a/Engine/Plugins/VirtualProduction/CameraCalibration/Source/CameraCalibrationEditor/Private/AssetEditor/CameraCalibrationStepsController.h b/Engine/Plugins/VirtualProduction/CameraCalibration/Source/CameraCalibrationEditor/Private/AssetEditor/CameraCalibrationStepsController.h
--- a/Engine/Plugins/VirtualProduction/CameraCalibration/Source/CameraCalibrationEditor/Private/AssetEditor/CameraCalibrationStepsController.h
+++ b/Engine/Plugins/VirtualProduction/CameraCalibration/Source/CameraCalibrationEditor/Private/AssetEditor/CameraCalibrationStepsController.h
@@ -19,6 +19,7 @@ class SWidget;
 class UMediaPlayer;
 class UMediaSource;
 class UMediaTexture;
+class UMediaProfile;
 class UCameraCalibrationStep;
 class UCompositingElementMaterialPass;
 class ULensComponent;
@@ -58,26 +59,20 @@ public:
 	/** Returns the UI that this object controls */
 	TSharedPtr<SWidget> BuildUI();
 
-	/** Creates composite with CG and selected media source */
-	void CreateComp();
+	/** Creates the media player and media texture used to play the media source */
+	void InitializeMediaPlayer();
 
-	/** Returns the render target of the Comp */
-	UTextureRenderTarget2D* GetRenderTarget() const;
+	/** Gets the texture used to display the media overlay */
+	UMediaTexture* GetMediaOverlayTexture() const;
 
-	/** Returns the render target of the Media Plate */
-	UTextureRenderTarget2D* GetMediaPlateRenderTarget() const;
+	/** Gets the resolution of the media overlay */
+	FIntPoint GetMediaOverlayResolution() const;
 
-	/** Returns the output resolution of the top-level Comp */
-	FIntPoint GetCompRenderResolution() const;
-
-	/** Returns the output resolution of the CG Layer */
-	FIntPoint GetCGRenderResolution() const;
-
-	/** Returns the size of the camera feed (which may be smaller than the top-level Comp and CG Layer) */
+	/** Returns the size of the camera feed (which may be smaller than the media overlay or the CG render) */
 	FIntPoint GetCameraFeedSize() const;
 
-	/** Creates a way to read the media plate pixels for processing by any calibration step */
-	void CreateMediaPlateOutput();
+	/** Returns the aspect ratio of the camera feed */
+	float GetCameraFeedAspectRatio() const;
 
 	/** Returns the CG weight that is composited on top of the media */
 	float GetWiperWeight() const;
@@ -121,6 +116,9 @@ public:
 	/** Sets the current media to a media texture being used in the calibration */
 	bool SetMediaTexture(UMediaTexture* InMediaTexture);
 
+	/** Sets the current media to a media source in the active media profile */
+	bool SetMediaProfileSource(UMediaSource* InMediaProfileSource);
+	
 	/** Clears out any media source or media texture currently being played */
 	void ClearMedia();
 
@@ -139,8 +137,17 @@ public:
 	/** Gets the current media source being played, or nullptr if nothing is being played */
 	UMediaSource* GetMediaSource() const;
 
+	/** Gets the media source from the active media profile currently being played, or nullptr if no media profile source is being played */
+	UMediaSource* GetMediaProfileSource() const;
+	
 	/** Gets the current media texture being displayed, or nullptr if the media texture is not being overridden */
 	UMediaTexture* GetMediaTexture() const;
+
+	/**
+	 * Gets the active media player being used; if an external media texture is being used,
+	 * this returns the player that texture is linked to, otherwise, the internal media player is returned
+	 */
+	UMediaPlayer* GetMediaPlayer() const;
 	
 	/** Returns the calibration steps */
 	const TConstArrayView<TStrongObjectPtr<UCameraCalibrationStep>> GetCalibrationSteps() const;
@@ -154,7 +161,7 @@ public:
 	/** Finds the world being used by the tool for finding and spawning objects */
 	UWorld* GetWorld() const;
 
-	/** Reads the pixels in the media plate */
+	/** Reads the pixels from the open media source */
 	bool ReadMediaPixels(TArray<FColor>& Pixels, FIntPoint& Size, FText& OutErrorMessage, ESimulcamViewportPortion ViewportPortion = ESimulcamViewportPortion::FullViewport) const;
 
 	/** Returns true if the overlay transform pass is currently enabled */
@@ -163,18 +170,21 @@ public:
 	/** Sets the enabled state of the overlay transform pass */
 	void SetOverlayEnabled(const bool bEnabled = true, EOverlayPassType OverlayPass = EOverlayPassType::ToolOverlay);
 
+	/** Returns the overlay material used by the input overlay pass type */
+	UMaterialInterface* GetOverlayMaterial(EOverlayPassType OverlayPass) const;
+	
 	/** Sets the overlay material to be used by the overlay transform pass */
 	void SetOverlayMaterial(UMaterialInterface* OverlayMaterial, bool bShowOverlay = true, EOverlayPassType OverlayPass = EOverlayPassType::ToolOverlay);
 
-	/** Redraw the overlay material used by the input overlay pass */
-	void RefreshOverlay(EOverlayPassType OverlayPass = EOverlayPassType::ToolOverlay);
-
 	/** Use the input mouse position (representing any corner of the camera feed) and the dimensions of the media source to calculate the new size of the camera feed */
 	void SetCameraFeedDimensionsFromMousePosition(FVector2D MousePosition);
 
 	/** Set the camera feed dimensions info of the LensFile being edited */
 	void SetCameraFeedDimensions(FIntPoint Dimensions, bool bMarkAsOverridden);
 
+	/** Resets the camera feed dimensions by fitting it within the currently playing media source */
+	void ResetCameraFeedDimensions();
+
 public:
 
 	/** Called by the UI when the Simulcam Viewport is clicked */
@@ -233,23 +243,14 @@ private:
 	/** Returns the first lens component with a matching LensFile found on the input camera, or nullptr if none exists */
 	ULensComponent* FindLensComponentOnCamera(ACameraActor* CineCamera) const;
 
-	/** Returns a namespaced version of the given name. Useful to generate names unique to this lens file */
-	FString NamespacedName(const FString&& Name) const;
-
-	/** Finds an existing comp element based on its name. */
-	ACompositingElement* FindElement(const FString& Name) const;
-
-	/** Spawns a new compositing element of the given class and parent element */
-	TWeakObjectPtr<ACompositingElement> AddElement(ACompositingElement* Parent, FString& ClassPath, FString& ElementName) const;
-
 	/** Releases resources used by the tool */
 	void Cleanup();
 
 	/** Convenience function that returns the first camera it finds that is using the lens associated with this object. */
 	ACameraActor* FindFirstCameraWithCurrentLens() const;
 
-	/** Enables distortion in the CG comp */
-	void EnableDistortionInCG();
+	/** When no camera can be found using the current lens file, this function finds a default camera in the level to use */
+	ACameraActor* FindDefaultCamera() const;
 
 	/** Called by the core ticker */
 	bool OnTick(float DeltaTime);
@@ -257,29 +258,11 @@ private:
 	/** Finds and creates the available calibration steps */
 	void CreateSteps();
 
-	/** Create a new material transform pass to represent an overlay and add it to the comp */
-	void CreateOverlayPass(FName PassName, TWeakObjectPtr<UCompositingElementMaterialPass>& OverlayPass, TWeakObjectPtr<UTextureRenderTarget2D>& OverlayRenderTarget);
-
-	/** Returns the overlay material pass associated with the input overlay pass type */
-	UCompositingElementMaterialPass* GetOverlayMaterialPass(EOverlayPassType OverlayPassType) const;
-
-	/** Returns the overlay render target used by the input overlay pass type */
-	UTextureRenderTarget2D* GetOverlayRenderTarget(EOverlayPassType OverlayPassType) const;
-
-	/** Returns the overlay material used by the input overlay pass type */
-	UMaterialInterface* GetOverlayMaterial(EOverlayPassType OverlayPass) const;
-
 	/** Iteratively attempt to resize the camera feed until its aspect ratio matches the input camera aspect ratio */
 	void MinimizeAspectRatioError(FIntPoint& CameraFeedDimensions, float CameraAspectRatio);
 
-	/** Resizes the resolution of the top-level comp match the resolution of the media source */
-	bool UpdateCompResolution();
-
-	/** Resizes the resolution of the CG layer to match the aspect resolution of the source camera and fit within the comp */
-	void UpdateCGResolution(bool bCompResized);
-
-	/** Update the aspect ratio correction material parameters based on the ratio of the camera feed dimensions to the comp dimensions */
-	void UpdateAspectRatioCorrection();
+	/** Raised when the active media profile is changed */
+	void OnActiveMediaProfileChanged(UMediaProfile* PreviousMediaProfile, UMediaProfile* NewMediaProfile);
 
 private:
 
@@ -292,20 +275,8 @@ private:
 	/** The lens asset */
 	TWeakObjectPtr<class ULensFile> LensFile;
 
-	/** The parent comp element */
-	TWeakObjectPtr<ACompositingElement> Comp;
-
-	/** The CG layer comp element */
-	TWeakObjectPtr<ACompositingElement> CGLayer;
-
-	/** The MediaPlate comp element*/
-	TWeakObjectPtr<ACompositingElement> MediaPlate;
-
-	/** The output render target of the comp */
-	TWeakObjectPtr<UTextureRenderTarget2D> RenderTarget;
-
-	/** The output render target of the media plate */
-	TWeakObjectPtr<UTextureRenderTarget2D> MediaPlateRenderTarget;
+	/** The index of the media source in the active media profile being played */
+	int32 MediaProfileSourceIndex = INDEX_NONE;
 
 	/** The media texture used by the media plate */
 	TWeakObjectPtr<UMediaTexture> MediaTexture;
@@ -314,29 +285,29 @@ private:
 	TWeakObjectPtr<UMediaTexture> ExternalMediaTexture;
 	
 	/** The media player that is playing the selected media source */
-	TStrongObjectPtr<UMediaPlayer> MediaPlayer;
-
-	/** The material pass the does the CG + MediaPlate composite with a wiper weight */
-	TWeakObjectPtr<UCompositingElementMaterialPass> MaterialPass;
+	TStrongObjectPtr<UMediaPlayer> InternalMediaPlayer;
 
-	/** A material pass set by one of the calibration steps that renders an overlay on top of the composite */
-	TWeakObjectPtr<UCompositingElementMaterialPass> ToolOverlayPass;
-
-	/** A material pass set by the user (via the editor UI) that renders an overlay on top of the composite */
-	TWeakObjectPtr<UCompositingElementMaterialPass> UserOverlayPass;
+	/** Used to extract pixel data from the media texture, cached to avoid having to recreate it every time */
+	TStrongObjectPtr<UTextureRenderTarget2D> CachedMediaRenderTarget;
 
 	/** The material used to render the overlay for the tool overlay pass */
 	TWeakObjectPtr<UMaterialInterface> ToolOverlayMaterial;
 
+	/** Whether the tool overlay is enabled */
+	bool bToolOverlayEnabled = false;
+	
 	/** The material used to render the overlay for the user overlay pass */
 	TWeakObjectPtr<UMaterialInterface> UserOverlayMaterial;
 
-	/** The render target used as input to the tool overlay material */
-	TWeakObjectPtr<UTextureRenderTarget2D> ToolOverlayRenderTarget;
+	/** Whether the user overlay is enabled */
+	bool bUserOverlayEnabled = false;
 
-	/** The render target used as input to the user overlay material */
-	TWeakObjectPtr<UTextureRenderTarget2D> UserOverlayRenderTarget;
+	/** Wiper weight for overlaying the media texture on the CG render */
+	float WiperWeight = 0.5f;
 
+	/** Keeps track of the last read media overlay resolution so that resolution changes can be detected */
+	FIntPoint LastMediaOverlayResolution;
+	
 	/** The currently selected camera */
 	TWeakObjectPtr<ACameraActor> Camera;
 
diff --git a/Engine/Plugins/VirtualProduction/CameraCalibration/Source/CameraCalibrationEditor/Private/AssetEditor/LensDistortionTool.cpp b/Engine/Plugins/VirtualProduction/CameraCalibration/Source/CameraCalibrationEditor/Private/AssetEditor/LensDistortionTool.cpp
--- a/Engine/Plugins/VirtualProduction/CameraCalibration/Source/CameraCalibrationEditor/Private/AssetEditor/LensDistortionTool.cpp
+++ b/Engine/Plugins/VirtualProduction/CameraCalibration/Source/CameraCalibrationEditor/Private/AssetEditor/LensDistortionTool.cpp
@@ -84,9 +84,15 @@ void ULensDistortionTool::Initialize(TWeakPtr<FCameraCalibrationStepsController>
 		UMaterialInterface* OverlayParent = TSoftObjectPtr<UMaterialInterface>(FSoftObjectPath(TEXT("/CameraCalibration/Materials/M_Coverage.M_Coverage"))).LoadSynchronous();
 		OverlayMID = UMaterialInstanceDynamic::Create(OverlayParent, GetTransientPackage());
 
-		const FIntPoint CompResolution = StepsController->GetCompRenderResolution();
-
-		OverlayTexture = UTexture2D::CreateTransient(CompResolution.X, CompResolution.Y, EPixelFormat::PF_B8G8R8A8);
+		FIntPoint OverlayResolution = StepsController->GetMediaOverlayResolution();
+		if (OverlayResolution == FIntPoint::ZeroValue)
+		{
+			// If the media overlay hasn't been initialized yet, it could return a zero size, in which case we will use a default size to initialize the texture
+			static const FIntPoint DefaultSize(1920, 1080);
+			OverlayResolution = DefaultSize;
+		}
+		
+		OverlayTexture = UTexture2D::CreateTransient(OverlayResolution.X, OverlayResolution.Y, EPixelFormat::PF_B8G8R8A8);
 		UE::CameraCalibration::Private::ClearTexture(OverlayTexture);
 
 		OverlayMID->SetTextureParameterValue(FName(TEXT("CoverageTexture")), OverlayTexture);
@@ -115,13 +121,14 @@ TSharedRef<SWidget> ULensDistortionTool::BuildUI()
 
 void ULensDistortionTool::Tick(float DeltaTime)
 {
-	// If the resolution of the simulcam comp has changed, update the coverage texture to be the correct size
+	// If the resolution of the media overlay texture has changed, update the coverage texture to be the correct size
 	if (TSharedPtr<FCameraCalibrationStepsController> StepsController = WeakStepsController.Pin())
 	{
-		const FIntPoint Size = StepsController->GetCompRenderResolution();
-		if (OverlayTexture && ((OverlayTexture->GetSizeX() != Size.X) || (OverlayTexture->GetSizeY() != Size.Y)))
+		const FIntPoint MediaOverlayResolution = StepsController->GetMediaOverlayResolution();
+		const FIntPoint TextureResolution = OverlayTexture ? FIntPoint(OverlayTexture->GetSizeX(), OverlayTexture->GetSizeY()) : FIntPoint::ZeroValue;
+		if (MediaOverlayResolution != TextureResolution)
 		{
-			OverlayTexture = UTexture2D::CreateTransient(Size.X, Size.Y, EPixelFormat::PF_B8G8R8A8);
+			OverlayTexture = UTexture2D::CreateTransient(MediaOverlayResolution.X, MediaOverlayResolution.Y, EPixelFormat::PF_B8G8R8A8);
 			RefreshCoverage();
 		}
 	}
@@ -373,11 +380,6 @@ bool ULensDistortionTool::DetectCheckerboardPattern(TArray<FColor>& Pixels, FInt
 
 	FOpenCVHelper::DrawCheckerboardCorners(CameraFeedAdjustedCorners, OutRow->CheckerboardDimensions, OverlayTexture);
 
-	if (TSharedPtr<FCameraCalibrationStepsController> StepsController = WeakStepsController.Pin())
-	{
-		StepsController->RefreshOverlay();
-	}
-
 	OutRow->TargetPose.SetLocation(TopLeft);
 
 	FRotator BoardRotation = FRotationMatrix::MakeFromYZ(TopRight - TopLeft, TopLeft - BottomLeft).Rotator();
@@ -458,11 +460,6 @@ bool ULensDistortionTool::DetectArucoPattern(TArray<FColor>& Pixels, FIntPoint S
 
 	FOpenCVHelper::DrawArucoMarkers(IdentifiedMarkers, OverlayTexture);
 
-	if (TSharedPtr<FCameraCalibrationStepsController> StepsController = WeakStepsController.Pin())
-	{
-		StepsController->RefreshOverlay();
-	}
-
 	return true;
 }
 
@@ -869,11 +866,6 @@ void ULensDistortionTool::RefreshCoverage()
 	{
 		OverlayMID->SetTextureParameterValue(FName(TEXT("CoverageTexture")), OverlayTexture);
 	}
-
-	if (TSharedPtr<FCameraCalibrationStepsController> StepsController = WeakStepsController.Pin())
-	{
-		StepsController->RefreshOverlay();
-	}
 }
 
 void ULensDistortionTool::RescalePoints(TArray<FVector2D>& Points, FIntPoint DebugTextureSize, FIntPoint CameraFeedSize)
diff --git a/Engine/Plugins/VirtualProduction/CameraCalibration/Source/CameraCalibrationEditor/Private/AssetEditor/LensInfoStep.cpp b/Engine/Plugins/VirtualProduction/CameraCalibration/Source/CameraCalibrationEditor/Private/AssetEditor/LensInfoStep.cpp
--- a/Engine/Plugins/VirtualProduction/CameraCalibration/Source/CameraCalibrationEditor/Private/AssetEditor/LensInfoStep.cpp
+++ b/Engine/Plugins/VirtualProduction/CameraCalibration/Source/CameraCalibrationEditor/Private/AssetEditor/LensInfoStep.cpp
@@ -269,9 +269,9 @@ void ULensInfoStep::OnCameraFeedInfoChanged(const FPropertyChangedEvent& Propert
 			if (ULensFile* LensFile = StepsController->GetLensFile())
 			{
 				FIntPoint CameraFeedDimensions = LensFile->CameraFeedInfo.GetDimensions();
-				const FIntPoint CompRenderResolution = StepsController->GetCompRenderResolution();
-				CameraFeedDimensions.X = FMath::Clamp(CameraFeedDimensions.X, 0, CompRenderResolution.X);
-				CameraFeedDimensions.Y = FMath::Clamp(CameraFeedDimensions.Y, 0, CompRenderResolution.Y);
+				const FIntPoint OverlayResolution = StepsController->GetMediaOverlayResolution();
+				CameraFeedDimensions.X = FMath::Clamp(CameraFeedDimensions.X, 0, OverlayResolution.X);
+				CameraFeedDimensions.Y = FMath::Clamp(CameraFeedDimensions.Y, 0, OverlayResolution.Y);
 
 				StepsController->SetCameraFeedDimensions(CameraFeedDimensions, true);
 			}
@@ -283,8 +283,7 @@ void ULensInfoStep::ResetCameraFeedInfoToDefault()
 {
 	if (TSharedPtr<FCameraCalibrationStepsController> StepsController = CameraCalibrationStepsController.Pin())
 	{
-		const FIntPoint DefaultCameraFeedDimensions = StepsController->GetCGRenderResolution();
-		StepsController->SetCameraFeedDimensions(DefaultCameraFeedDimensions, false);
+		StepsController->ResetCameraFeedDimensions();
 	}
 }
 
diff --git a/Engine/Plugins/VirtualProduction/CameraCalibration/Source/CameraCalibrationEditor/Private/AssetEditor/SCameraCalibrationSteps.cpp b/Engine/Plugins/VirtualProduction/CameraCalibration/Source/CameraCalibrationEditor/Private/AssetEditor/SCameraCalibrationSteps.cpp
--- a/Engine/Plugins/VirtualProduction/CameraCalibration/Source/CameraCalibrationEditor/Private/AssetEditor/SCameraCalibrationSteps.cpp
+++ b/Engine/Plugins/VirtualProduction/CameraCalibration/Source/CameraCalibrationEditor/Private/AssetEditor/SCameraCalibrationSteps.cpp
@@ -239,7 +239,67 @@ void SCameraCalibrationSteps::Construct(const FArguments& InArgs, TWeakPtr<FCame
 			+ SVerticalBox::Slot() // Simulcam Viewport
 			.VAlign(EVerticalAlignment::VAlign_Fill)
 			[
-				SNew(SSimulcamViewport, CalibrationStepsController.Pin()->GetRenderTarget())
+				SNew(SSimulcamViewport)
+				.ViewActor_Lambda([WeakStepsController = CalibrationStepsController]()->ACameraActor*
+				{
+					if (TSharedPtr<FCameraCalibrationStepsController> PinnedStepsController = WeakStepsController.Pin())
+					{
+						return PinnedStepsController->GetCamera();
+					}
+
+					return nullptr;
+				})
+				.FilmbackAspectRatio_Lambda([WeakStepsController = CalibrationStepsController]
+				{
+					if (TSharedPtr<FCameraCalibrationStepsController> PinnedStepsController = WeakStepsController.Pin())
+					{
+						return PinnedStepsController->GetCameraFeedAspectRatio();
+					}
+
+					return 1.0f;
+				})
+				.MediaOverlayTexture_Lambda([WeakStepsController = CalibrationStepsController]()->UTexture*
+				{
+					if (TSharedPtr<FCameraCalibrationStepsController> PinnedStepsController = WeakStepsController.Pin())
+					{
+						return PinnedStepsController->GetMediaOverlayTexture();
+					}
+
+					return nullptr;
+				})
+				.MediaOverlayOpacity_Lambda([WeakStepsController = CalibrationStepsController]
+				{
+					if (TSharedPtr<FCameraCalibrationStepsController> PinnedStepsController = WeakStepsController.Pin())
+					{
+						return 1.0f - PinnedStepsController->GetWiperWeight();
+					}
+
+					return 0.0f;
+				})
+				.ToolOverlayMaterial_Lambda([WeakStepsController = CalibrationStepsController]()->UMaterialInterface*
+				{
+					if (TSharedPtr<FCameraCalibrationStepsController> PinnedStepsController = WeakStepsController.Pin())
+					{
+						if (PinnedStepsController->IsOverlayEnabled(EOverlayPassType::ToolOverlay))
+						{
+							return PinnedStepsController->GetOverlayMaterial(EOverlayPassType::ToolOverlay);
+						}
+					}
+
+					return nullptr;
+				})
+				.UserOverlayMaterial_Lambda([WeakStepsController = CalibrationStepsController]()->UMaterialInterface*
+				{
+					if (TSharedPtr<FCameraCalibrationStepsController> PinnedStepsController = WeakStepsController.Pin())
+					{
+						if (PinnedStepsController->IsOverlayEnabled(EOverlayPassType::UserOverlay))
+						{
+							return PinnedStepsController->GetOverlayMaterial(EOverlayPassType::UserOverlay);
+						}
+					}
+
+					return nullptr;
+				})
 				.OnSimulcamViewportClicked_Raw(CalibrationStepsController.Pin().Get(), &FCameraCalibrationStepsController::OnSimulcamViewportClicked)
 				.OnSimulcamViewportInputKey_Raw(CalibrationStepsController.Pin().Get(), &FCameraCalibrationStepsController::OnSimulcamViewportInputKey)
 				.OnSimulcamViewportMarqueeSelect_Raw(CalibrationStepsController.Pin().Get(), &FCameraCalibrationStepsController::OnSimulcamViewportMarqueeSelect)
@@ -616,7 +676,7 @@ TSharedRef<SWidget> SCameraCalibrationSteps::BuildMediaProfileSourcePicker()
 				return;
 			}
 
-			CalibrationStepsController.Pin()->SetMediaSource(NewValue.Get());
+			CalibrationStepsController.Pin()->SetMediaProfileSource(NewValue.Get());
 		})
 		.OnGenerateWidget_Lambda([&](TWeakObjectPtr<UMediaSource> InMediaSource) -> TSharedRef<SWidget>
 		{
@@ -638,7 +698,7 @@ TSharedRef<SWidget> SCameraCalibrationSteps::BuildMediaProfileSourcePicker()
 					return None;
 				}
 
-				UMediaSource* MediaSource = CalibrationStepsController.Pin()->GetMediaSource();
+				UMediaSource* MediaSource = CalibrationStepsController.Pin()->GetMediaProfileSource();
 				if (!MediaSource)
 				{
 					return None;
@@ -796,11 +856,6 @@ void SCameraCalibrationSteps::UpdateOverlayMaterialParameterWidget()
 							if (UMaterialInstanceDynamic* Overlay = CurrentOverlayMID.Get())
 							{
 								Overlay->SetScalarParameterValue(ParameterInfo.Name, NewValue);
-
-								if (CalibrationStepsController.IsValid())
-								{
-									CalibrationStepsController.Pin()->RefreshOverlay(EOverlayPassType::UserOverlay);
-								}
 							}
 						})
 						.OnValueCommitted_Lambda([this, ParameterInfo](float NewValue, ETextCommit::Type)
@@ -808,11 +863,6 @@ void SCameraCalibrationSteps::UpdateOverlayMaterialParameterWidget()
 							if (UMaterialInstanceDynamic* Overlay = CurrentOverlayMID.Get())
 							{
 								Overlay->SetScalarParameterValue(ParameterInfo.Name, NewValue);
-
-								if (CalibrationStepsController.IsValid())
-								{
-									CalibrationStepsController.Pin()->RefreshOverlay(EOverlayPassType::UserOverlay);
-								}
 							}
 						})
 					]
@@ -864,11 +914,6 @@ void SCameraCalibrationSteps::UpdateOverlayMaterialParameterWidget()
 									if (UMaterialInstanceDynamic* Overlay = CurrentOverlayMID.Get())
 									{
 										Overlay->SetVectorParameterValue(ParameterInfo.Name, NewColor);
-
-										if (CalibrationStepsController.IsValid())
-										{
-											CalibrationStepsController.Pin()->RefreshOverlay(EOverlayPassType::UserOverlay);
-										}
 									}
 								});
 						})
diff --git a/Engine/Plugins/VirtualProduction/CameraCalibration/Source/CameraCalibrationEditor/Private/AssetEditor/SSimulcamViewport.cpp b/Engine/Plugins/VirtualProduction/CameraCalibration/Source/CameraCalibrationEditor/Private/AssetEditor/SSimulcamViewport.cpp
--- a/Engine/Plugins/VirtualProduction/CameraCalibration/Source/CameraCalibrationEditor/Private/AssetEditor/SSimulcamViewport.cpp
+++ b/Engine/Plugins/VirtualProduction/CameraCalibration/Source/CameraCalibrationEditor/Private/AssetEditor/SSimulcamViewport.cpp
@@ -5,14 +5,19 @@
 #include "Viewport/SSimulcamEditorViewport.h"
 
 
-void SSimulcamViewport::Construct(const FArguments& InArgs, UTexture* InTexture)
+void SSimulcamViewport::Construct(const FArguments& InArgs)
 {
+	ViewActor = InArgs._ViewActor;
+	FilmbackAspectRatio = InArgs._FilmbackAspectRatio;
+	MediaOverlayTexture = InArgs._MediaOverlayTexture;
+	MediaOverlayOpacity = InArgs._MediaOverlayOpacity;
+	ToolOverlayMaterial = InArgs._ToolOverlayMaterial;
+	UserOverlayMaterial = InArgs._UserOverlayMaterial;
+	OverrideTexture = TStrongObjectPtr<UTexture>(InArgs._OverrideTexture);
 	OnSimulcamViewportClicked = InArgs._OnSimulcamViewportClicked;
 	OnSimulcamViewportInputKey = InArgs._OnSimulcamViewportInputKey;
 	OnSimulcamViewportMarqueeSelect = InArgs._OnSimulcamViewportMarqueeSelect;
 
-	Texture = TStrongObjectPtr<UTexture>(InTexture);
-
 	TextureViewport = SNew(SSimulcamEditorViewport, SharedThis(this), InArgs._WithZoom.Get(), InArgs._WithPan.Get());
 
 	ChildSlot
@@ -21,15 +26,39 @@ void SSimulcamViewport::Construct(const FArguments& InArgs, UTexture* InTexture)
 		];
 }
 
-UTexture* SSimulcamViewport::GetTexture() const
+ACameraActor* SSimulcamViewport::GetCameraActor() const
+{
+	return ViewActor.Get(nullptr);
+}
+
+float SSimulcamViewport::GetCameraFeedAspectRatio() const
+{
+	return FilmbackAspectRatio.Get(1.0f);
+}
+
+UTexture* SSimulcamViewport::GetMediaOverlayTexture() const
+{
+	return MediaOverlayTexture.Get(nullptr);
+}
+
+float SSimulcamViewport::GetMediaOverlayOpacity() const
+{
+	return MediaOverlayOpacity.Get(1.0f);
+}
+
+UMaterialInterface* SSimulcamViewport::GetToolOverlayMaterial() const
+{
+	return ToolOverlayMaterial.Get(nullptr);
+}
+
+UMaterialInterface* SSimulcamViewport::GetUserOverlayMaterial() const
 {
-	return Texture.Get();
+	return UserOverlayMaterial.Get(nullptr);
 }
 
-bool SSimulcamViewport::HasValidTextureResource() const
+UTexture* SSimulcamViewport::GetOverrideTexture() const
 {
-	UTexture* CurrentTexture = GetTexture();
-	return CurrentTexture != nullptr && CurrentTexture->GetResource() != nullptr;
+	return OverrideTexture.Get();
 }
 
 bool SSimulcamViewport::OnViewportInputKey(const FKey& Key, const EInputEvent& Event)
diff --git a/Engine/Plugins/VirtualProduction/CameraCalibration/Source/CameraCalibrationEditor/Private/AssetEditor/SSimulcamViewport.h b/Engine/Plugins/VirtualProduction/CameraCalibration/Source/CameraCalibrationEditor/Private/AssetEditor/SSimulcamViewport.h
--- a/Engine/Plugins/VirtualProduction/CameraCalibration/Source/CameraCalibrationEditor/Private/AssetEditor/SSimulcamViewport.h
+++ b/Engine/Plugins/VirtualProduction/CameraCalibration/Source/CameraCalibrationEditor/Private/AssetEditor/SSimulcamViewport.h
@@ -7,6 +7,7 @@
 #include "Engine/EngineBaseTypes.h"
 #include "UObject/GCObject.h"
 #include "UObject/StrongObjectPtr.h"
+#include "Viewport/SimulcamEditorViewportClient.h"
 #include "Widgets/DeclarativeSyntaxSupport.h"
 
 class UTexture;
@@ -27,45 +28,78 @@ DECLARE_DELEGATE_RetVal_TwoParams(bool, FSimulcamViewportInputKeyEventHandler, c
 DECLARE_DELEGATE_TwoParams(FSimulcamViewportMarqueeSelectEventHandler, FVector2D StartPosition, FVector2D EndPosition);
 
 /**
- * UI to display the provided UTexture and respond to user input.
+ * Viewport that displays a level render from the perspective of a camera actor overlaid with a media texture overlay
  */
-class SSimulcamViewport : public SCompoundWidget
+class SSimulcamViewport : public SCompoundWidget, public FSimulcamViewportElementsProvider
 {
 public:
 	SLATE_BEGIN_ARGS(SSimulcamViewport)
-	: _WithZoom(true), _WithPan(true)
+		:
+		_ViewActor(nullptr),
+		_FilmbackAspectRatio(1.0f),
+		_MediaOverlayTexture(nullptr),
+		_MediaOverlayOpacity(0.5f),
+		_ToolOverlayMaterial(nullptr),
+		_UserOverlayMaterial(nullptr),
+		_OverrideTexture(nullptr),
+		_WithZoom(true),
+		_WithPan(true)
 	{}
-	SLATE_EVENT(FSimulcamViewportClickedEventHandler, OnSimulcamViewportClicked)
-	SLATE_EVENT(FSimulcamViewportInputKeyEventHandler, OnSimulcamViewportInputKey)
-	SLATE_EVENT(FSimulcamViewportMarqueeSelectEventHandler, OnSimulcamViewportMarqueeSelect)
-	SLATE_ATTRIBUTE(bool, WithZoom)
-	SLATE_ATTRIBUTE(bool, WithPan)
+		SLATE_ATTRIBUTE(ACameraActor*, ViewActor)
+		SLATE_ATTRIBUTE(float, FilmbackAspectRatio)
+		SLATE_ATTRIBUTE(UTexture*, MediaOverlayTexture)
+		SLATE_ATTRIBUTE(float, MediaOverlayOpacity)
+		SLATE_ATTRIBUTE(UMaterialInterface*, ToolOverlayMaterial)
+		SLATE_ATTRIBUTE(UMaterialInterface*, UserOverlayMaterial)
+		SLATE_ARGUMENT(UTexture*, OverrideTexture)
+		SLATE_EVENT(FSimulcamViewportClickedEventHandler, OnSimulcamViewportClicked)
+		SLATE_EVENT(FSimulcamViewportInputKeyEventHandler, OnSimulcamViewportInputKey)
+		SLATE_EVENT(FSimulcamViewportMarqueeSelectEventHandler, OnSimulcamViewportMarqueeSelect)
+		SLATE_ATTRIBUTE(bool, WithZoom)
+		SLATE_ATTRIBUTE(bool, WithPan)
 	SLATE_END_ARGS()
 
 	/**
 	 * Construct this widget
 	 *
 	 * @param InArgs The declaration data for this widget.
-	 * @param InTexture The source texture to render if any
 	 */
-	void Construct(const FArguments& InArgs, UTexture* InTexture);
+	void Construct(const FArguments& InArgs);
+
+	//~ Begin FSimulcamViewportElementsProvider interface
+	virtual  ACameraActor* GetCameraActor() const override;
+	virtual float GetCameraFeedAspectRatio() const override;
+	virtual UTexture* GetMediaOverlayTexture() const override;
+	virtual float GetMediaOverlayOpacity() const override;
+	virtual UMaterialInterface* GetToolOverlayMaterial() const override;
+	virtual UMaterialInterface* GetUserOverlayMaterial() const override;
+	virtual UTexture* GetOverrideTexture() const override;
+	virtual void OnViewportClicked(const FGeometry& MyGeometry, const FPointerEvent& PointerEvent) override { OnSimulcamViewportClicked.ExecuteIfBound(MyGeometry, PointerEvent); }
+	virtual bool OnViewportInputKey(const FKey& InKey, const EInputEvent& InInputEvent) override;
+	virtual void OnMarqueeSelect(const FVector2D& InStartPosition, const FVector2D& InEndPosition) override { OnSimulcamViewportMarqueeSelect.ExecuteIfBound(InStartPosition, InEndPosition); }
+	//~ End FSimulcamViewportElementsProvider interface
 
-	/** Returns the current camera texture */
-	UTexture* GetTexture() const;
+private:
+	/** The camera actor that the scene is rendering from */
+	TAttribute<ACameraActor*> ViewActor;
 
-	/** Is the camera texture valid ? */
-	bool HasValidTextureResource() const;
+	/** The aspect ratio of the camera's filmback */
+	TAttribute<float> FilmbackAspectRatio;
 
-	/** Called whenever the viewport is clicked */
-	void OnViewportClicked(const FGeometry& MyGeometry, const FPointerEvent& PointerEvent) { OnSimulcamViewportClicked.ExecuteIfBound(MyGeometry, PointerEvent); }
+	/** The media texture to overlay onto the scene render */
+	TAttribute<UTexture*> MediaOverlayTexture;
 
-	/** Called when the viewport receives input key presses*/
-	bool OnViewportInputKey(const FKey& Key, const EInputEvent& Event);
+	/** The opacity of the media overlay */
+	TAttribute<float> MediaOverlayOpacity;
 
-	/** Called when the viewport receives a marquee select event */
-	void OnMarqueeSelect(FVector2D StartPosition, FVector2D EndPosition) { OnSimulcamViewportMarqueeSelect.ExecuteIfBound(StartPosition, EndPosition); }
+	/** The material to use to draw the tool overlay */
+	TAttribute<UMaterialInterface*> ToolOverlayMaterial;
 
-private:
+	/** The material to use to draw the user overlay */
+	TAttribute<UMaterialInterface*> UserOverlayMaterial;
+
+	/** An override texture that is displayed instead of the scene render and media overlay */
+	TStrongObjectPtr<UTexture> OverrideTexture;
 
 	/** Delegate to be executed when the viewport area is clicked */
 	FSimulcamViewportClickedEventHandler OnSimulcamViewportClicked;
@@ -77,7 +111,5 @@ private:
 	FSimulcamViewportMarqueeSelectEventHandler OnSimulcamViewportMarqueeSelect;
 
 	TSharedPtr<SSimulcamEditorViewport> TextureViewport;
-
-	TStrongObjectPtr<UTexture> Texture;
 };
 
diff --git a/Engine/Plugins/VirtualProduction/CameraCalibration/Source/CameraCalibrationEditor/Private/AssetEditor/Viewport/SSimulcamEditorViewport.cpp b/Engine/Plugins/VirtualProduction/CameraCalibration/Source/CameraCalibrationEditor/Private/AssetEditor/Viewport/SSimulcamEditorViewport.cpp
--- a/Engine/Plugins/VirtualProduction/CameraCalibration/Source/CameraCalibrationEditor/Private/AssetEditor/Viewport/SSimulcamEditorViewport.cpp
+++ b/Engine/Plugins/VirtualProduction/CameraCalibration/Source/CameraCalibrationEditor/Private/AssetEditor/Viewport/SSimulcamEditorViewport.cpp
@@ -5,14 +5,9 @@
 #include "AssetEditor/SSimulcamViewport.h"
 #include "AssetEditor/Viewport/SimulcamEditorViewportClient.h"
 #include "Engine/Texture.h"
-#include "Engine/TextureRenderTargetVolume.h"
-#include "Engine/VolumeTexture.h"
 #include "Framework/Application/SlateApplication.h"
-#include "Framework/MultiBox/MultiBoxBuilder.h"
 #include "Slate/SceneViewport.h"
-#include "Widgets/Input/SComboButton.h"
 #include "Widgets/Input/SNumericEntryBox.h"
-#include "Widgets/Input/SSlider.h"
 #include "Widgets/SToolTip.h"
 #include "Widgets/SViewport.h"
 #include "Widgets/Text/STextBlock.h"
@@ -73,15 +68,23 @@ void SSimulcamEditorViewport::Construct(const FArguments& InArgs, const TSharedR
 
 	];
 
-	ViewportClient = MakeShared<FSimulcamEditorViewportClient>(InSimulcamViewport, SharedThis(this), bWithZoom, bWithPan);
+	ViewportClient = MakeShared<FSimulcamEditorViewportClient>(InSimulcamViewport, bWithZoom, bWithPan);
 	Viewport = MakeShared<FSceneViewport>(ViewportClient.Get(), ViewportWidget);
+	ViewportClient->Viewport = Viewport.Get();
 	// The viewport widget needs an interface so it knows what should render
 	ViewportWidget->SetViewportInterface(Viewport.ToSharedRef());
+}
 
-	// required for having the first 'zoom to fit' behaviour
-	CacheEffectiveTextureSize();
+SSimulcamEditorViewport::~SSimulcamEditorViewport()
+{
+	if (ViewportClient.IsValid())
+	{
+		ViewportClient->Viewport = nullptr;
+	}
 
-	Viewport->ViewportResizedEvent.AddSP(ViewportClient.ToSharedRef(), &FSimulcamEditorViewportClient::OnViewportResized);
+	Viewport.Reset();
+	ViewportWidget.Reset();
+	ViewportClient.Reset();
 }
 
 void SSimulcamEditorViewport::EnableRendering()
@@ -117,130 +120,30 @@ FText SSimulcamEditorViewport::GetDisplayedResolution() const
 	return ViewportClient->GetDisplayedResolution();
 }
 
-bool SSimulcamEditorViewport::HasValidTextureResource() const
-{
-	return (SimulcamViewportWeakPtr.IsValid() && SimulcamViewportWeakPtr.Pin()->HasValidTextureResource());
-}
-
 FText SSimulcamEditorViewport::HandleZoomPercentageText() const
 {
-	if (SimulcamViewportWeakPtr.IsValid() && HasValidTextureResource())
-	{
-		FText FormattedText = FText::FromString(TEXT("{0}: {1}"));
-		double DisplayedZoomLevel = GetCustomZoomLevel();
-		return FText::Format(FormattedText, FText::FromName(SimulcamViewportWeakPtr.Pin()->GetTexture()->GetFName()), FText::AsPercent(DisplayedZoomLevel));
-	}
-	else
-	{
-		return LOCTEXT("InvalidTexture", "Invalid Texture");
-	}
-}
-
-void SSimulcamEditorViewport::CacheEffectiveTextureSize()
-{
-	if (!SimulcamViewportWeakPtr.IsValid())
-	{
-		return;
-	}
-
-	UTexture* Texture = SimulcamViewportWeakPtr.Pin()->GetTexture();
-	if (!Texture)
-	{
-		return;
-	}
-
-	// If the cached texture dimensions do not match the current texture dimensions (because the texture was resized), update the cached value and zoom all the way out
-	if (CachedEffectiveTextureSize.X != Texture->GetSurfaceWidth() || CachedEffectiveTextureSize.Y != Texture->GetSurfaceHeight())
-	{
-		CachedEffectiveTextureSize = FVector2D(Texture->GetSurfaceWidth(), Texture->GetSurfaceHeight());
-		ViewportClient->OnTextureResized();
-	}
-}
-
-FVector2D SSimulcamEditorViewport::CalculateTextureDimensions() const
-{
-	if (!SimulcamViewportWeakPtr.IsValid())
+	UTexture* Texture = nullptr;
+	if (SimulcamViewportWeakPtr.IsValid())
 	{
-		return FVector2D::ZeroVector;
-	}
+		Texture = SimulcamViewportWeakPtr.Pin()->GetOverrideTexture();
 
-	UTexture* Texture = SimulcamViewportWeakPtr.Pin()->GetTexture();
-	if (!Texture)
-	{
-		return FVector2D::ZeroVector;
+		if (!Texture)
+		{
+			Texture = SimulcamViewportWeakPtr.Pin()->GetMediaOverlayTexture();
+		}
 	}
-
-	FVector2D Size = FVector2D(CachedEffectiveTextureSize.X * Zoom, CachedEffectiveTextureSize.Y * Zoom);
-
-	/** catch if the Width and Height are still zero for some reason */
-	if ((Size.X == 0) || (Size.Y == 0))
+	
+	if (Texture)
 	{
-		Size = FVector2D::ZeroVector;
-	}
-
-	return Size;
-}
-
-double SSimulcamEditorViewport::GetCustomZoomLevel() const
-{
-	return Zoom;
-}
-
-double SSimulcamEditorViewport::GetMinZoomLevel() const
-{
-	FIntPoint ViewportSize = GetViewport()->GetSizeXY();
-	return FMath::Min(ViewportSize.X / CachedEffectiveTextureSize.X, ViewportSize.Y / CachedEffectiveTextureSize.Y);
-}
-
-FVector2D SSimulcamEditorViewport::GetFitPosition() const
-{
-	FIntPoint ViewportSize = GetViewport()->GetSizeXY();
-	double MinZoom = GetMinZoomLevel();
-	return FVector2D((ViewportSize.X - (CachedEffectiveTextureSize.X * MinZoom)) / 2, (ViewportSize.Y - (CachedEffectiveTextureSize.Y * MinZoom)) / 2);
-}
-
-void SSimulcamEditorViewport::SetCustomZoomLevel(double ZoomValue)
-{
-	Zoom = FMath::Clamp(ZoomValue, GetMinZoomLevel(), MaxZoom);
-}
+		FText FormattedText = FText::FromString(TEXT("{0}: {1}"));
+		double DisplayedZoomLevel = ViewportClient->GetMediaOverlayScale();
 
-void SSimulcamEditorViewport::OffsetZoom(double OffsetValue, bool bSnapToStepSize)
-{
-	/** Offset from our current "visual" zoom level so that you can
-		smoothly transition from Fit/Fill mode into a custom zoom level */
-	const double CurrentZoom = GetCustomZoomLevel();
-	if (bSnapToStepSize)
-	{
-		/** Snap to the zoom step when offsetting to avoid zooming all the way to the min(0.01)
-			then back up (+0.1) causing your zoom level to be off by 0.01 (eg. 11%)
-			If we were in a fit view mode then our current zoom level could also be off the grid */
-		const double FinalZoom = FMath::GridSnap(CurrentZoom + OffsetValue, ZoomStep);
-		SetCustomZoomLevel(FinalZoom);
+		return FText::Format(FormattedText, FText::FromName(Texture->GetFName()), FText::AsPercent(DisplayedZoomLevel));
 	}
 	else
 	{
-		SetCustomZoomLevel(CurrentZoom + OffsetValue);
-	}
-}
-
-void SSimulcamEditorViewport::ZoomIn()
-{
-	OffsetZoom(ZoomStep);
-}
-
-void SSimulcamEditorViewport::ZoomOut()
-{
-	OffsetZoom(-ZoomStep);
-}
-
-void SSimulcamEditorViewport::OnViewportClicked(const FGeometry& MyGeometry, const FPointerEvent& PointerEvent)
-{
-	if (!SimulcamViewportWeakPtr.IsValid())
-	{
-		return;
+		return LOCTEXT("InvalidTexture", "Invalid Texture");
 	}
-
-	SimulcamViewportWeakPtr.Pin()->OnViewportClicked(MyGeometry, PointerEvent);
 }
 
 #undef LOCTEXT_NAMESPACE
diff --git a/Engine/Plugins/VirtualProduction/CameraCalibration/Source/CameraCalibrationEditor/Private/AssetEditor/Viewport/SSimulcamEditorViewport.h b/Engine/Plugins/VirtualProduction/CameraCalibration/Source/CameraCalibrationEditor/Private/AssetEditor/Viewport/SSimulcamEditorViewport.h
--- a/Engine/Plugins/VirtualProduction/CameraCalibration/Source/CameraCalibrationEditor/Private/AssetEditor/Viewport/SSimulcamEditorViewport.h
+++ b/Engine/Plugins/VirtualProduction/CameraCalibration/Source/CameraCalibrationEditor/Private/AssetEditor/Viewport/SSimulcamEditorViewport.h
@@ -1,13 +1,12 @@
 // Copyright Epic Games, Inc. All Rights Reserved.
 #pragma once
 #include "CoreMinimal.h"
-#include "Layout/Visibility.h"
+#include "SimulcamEditorViewportClient.h"
 #include "Widgets/DeclarativeSyntaxSupport.h"
 #include "Widgets/SCompoundWidget.h"
 class FSceneViewport;
 class SViewport;
 class SSimulcamViewport;
-class FSimulcamEditorViewportClient;
 struct FPointerEvent;
 
 /**
@@ -28,6 +27,8 @@ public:
 	 */
 	void Construct( const FArguments& InArgs, const TSharedRef<SSimulcamViewport>& InSimulcamViewport, const bool bWithZoom, const bool bWithPan);
 
+	virtual ~SSimulcamEditorViewport() override;
+	
 	/** Enable viewport rendering */
 	void EnableRendering();
 	/** Disable viewport rendering */
@@ -39,39 +40,10 @@ public:
 	/** Get The SViewport Widget using this editor */
 	TSharedPtr<SViewport> GetViewportWidget( ) const;
 
-	/** Returns the current texture dimensions (in pixels) */
-	FVector2D CalculateTextureDimensions() const;
-
-	/** The current zoom level */
-	double GetCustomZoomLevel() const;
-
-	/** Directly set a zoom level */
-	void SetCustomZoomLevel(double ZoomValue);
-
-	/** Zoom with offset */
-	void OffsetZoom(double OffsetValue, bool bSnapToStepSize = true);
-	
-	/** Increase zooming */
-	void ZoomIn();
-
-	/** Decrease zooming */
-	void ZoomOut();
-
 	//~ Begin SWidget interface
 	virtual void Tick( const FGeometry& AllottedGeometry, const double InCurrentTime, const float InDeltaTime ) override;
 	//~ End SWidget interface
 
-	/** Quick compute and cache the texture size before each frame */
-	void CacheEffectiveTextureSize();
-
-	/** Returns the minimal level of Zoom */
-	double GetMinZoomLevel() const;
-
-	/** Compute the best position for fitting the texture in the viewport */
-	FVector2D GetFitPosition() const;
-
-	/** Triggered whenever the viewport is clicked */
-	void OnViewportClicked(const FGeometry& MyGeometry, const FPointerEvent& PointerEvent);
 protected:
 	/**
 	 * Gets the displayed textures resolution as a string.
@@ -85,8 +57,6 @@ private:
 	// Callback for getting the zoom percentage text.
 	FText HandleZoomPercentageText( ) const;
 
-	// Checks if the texture being edited has a valid texture resource
-	bool HasValidTextureResource( ) const;
 private:
 	TWeakPtr<SSimulcamViewport> SimulcamViewportWeakPtr;
 	
@@ -99,14 +69,4 @@ private:
 
 	// Is rendering currently enabled? (disabled when reimporting a texture)
 	bool bIsRenderingEnabled;
-
-	/** The maximum width/height at which the texture will render in the preview window */
-	FVector2D CachedEffectiveTextureSize;
-
-	/** The texture's zoom factor. forced to 'fit' when 0 */
-	double Zoom = 0;
-
-public:
-	static constexpr double MaxZoom = 16.0;
-	static constexpr double ZoomStep = 0.1;
 };
diff --git a/Engine/Plugins/VirtualProduction/CameraCalibration/Source/CameraCalibrationEditor/Private/AssetEditor/Viewport/SimulcamEditorViewportClient.cpp b/Engine/Plugins/VirtualProduction/CameraCalibration/Source/CameraCalibrationEditor/Private/AssetEditor/Viewport/SimulcamEditorViewportClient.cpp
--- a/Engine/Plugins/VirtualProduction/CameraCalibration/Source/CameraCalibrationEditor/Private/AssetEditor/Viewport/SimulcamEditorViewportClient.cpp
+++ b/Engine/Plugins/VirtualProduction/CameraCalibration/Source/CameraCalibrationEditor/Private/AssetEditor/Viewport/SimulcamEditorViewportClient.cpp
@@ -1,133 +1,167 @@
 // Copyright Epic Games, Inc. All Rights Reserved.
 
 #include "AssetEditor/Viewport/SimulcamEditorViewportClient.h"
-#include "AssetEditor/SSimulcamViewport.h"
-#include "AssetEditor/Viewport/SSimulcamEditorViewport.h"
 #include "CanvasItem.h"
 #include "CanvasTypes.h"
-#include "Engine/Texture2D.h"
 #include "Editor/UnrealEdEngine.h"
-#include "EngineModule.h"
-#include "EngineUtils.h"
 #include "Framework/Application/SlateApplication.h"
 #include "ImageUtils.h"
+#include "SceneView.h"
 #include "Slate/SceneViewport.h"
-#include "Texture2DPreview.h"
 #include "TextureResource.h"
-#include "UnrealEdGlobals.h"
+#include "Camera/CameraActor.h"
+#include "Materials/MaterialInstanceDynamic.h"
 
-
-FSimulcamEditorViewportClient::FSimulcamEditorViewportClient(const TSharedRef<SSimulcamViewport>& InSimulcamViewport, const TSharedRef<SSimulcamEditorViewport>& InSimulcamEditorViewport, const bool bInWithZoom, const bool bInWithPan)
-	: SimulcamViewportWeakPtr(InSimulcamViewport)
-	, SimulcamEditorViewportWeakPtr(InSimulcamEditorViewport)
+FSimulcamEditorViewportClient::FSimulcamEditorViewportClient(TSharedPtr<FSimulcamViewportElementsProvider> InElementsProvider, const bool bInWithZoom, const bool bInWithPan)
+	: FLevelEditorViewportClient(nullptr)
+	, ElementsProvider(InElementsProvider)
 	, bWithZoom(bInWithZoom)
 	, bWithPan(bInWithPan)
 {
+	bIsRealtime = true;
+
+	VisibilityDelegate.BindLambda([] {return true; });
+
+	EngineShowFlags = ESFIM_Game;
+	EngineShowFlags.SetFog(false);
+	
+	FSoftObjectPath CropToViewRectMaterialPath("/CameraCalibration/Materials/M_CropToViewRect.M_CropToViewRect");
+	if (UMaterialInterface* CropToViewRectMaterial = Cast<UMaterialInterface>(CropToViewRectMaterialPath.TryLoad()))
+	{
+		CropToViewRectMaterialInstance = TStrongObjectPtr<UMaterialInstanceDynamic>(UMaterialInstanceDynamic::Create(CropToViewRectMaterial, nullptr));
+	}
 }
 
-void FSimulcamEditorViewportClient::Draw(FViewport* Viewport, FCanvas* Canvas)
+void FSimulcamEditorViewportClient::Tick(float DeltaSeconds)
 {
-	if (!SimulcamViewportWeakPtr.IsValid() || !SimulcamEditorViewportWeakPtr.IsValid())
+	ACameraActor* ViewActor = ElementsProvider.IsValid() ? ElementsProvider.Pin()->GetCameraActor() : nullptr;
+	SetActorLock(ViewActor);
+	
+	FLevelEditorViewportClient::Tick(DeltaSeconds);
+
+	// FLevelEditorViewportClient will update the post-process blends from the controlling actor every tick,
+	// so append the crop to view rect pass afterward
+	FPostProcessSettings CropToViewRectPPSettings;
+	CropToViewRectPPSettings.AddBlendable(CropToViewRectMaterialInstance.Get(), 1.0f);
+	
+	ControllingActorExtraPostProcessBlends.Add(CropToViewRectPPSettings);
+	ControllingActorExtraPostProcessBlendWeights.Add(1.0f);
+}
+
+void FSimulcamEditorViewportClient::Draw(FViewport* InViewport, FCanvas* Canvas)
+{	
+	UTexture* OverrideTexture = ElementsProvider.IsValid() ? ElementsProvider.Pin()->GetOverrideTexture() : nullptr;
+	if (OverrideTexture && OverrideTexture->GetResource())
 	{
+		ViewportSize = FVector2D(InViewport->GetSizeXY());
+		
+		const FVector2D TextureSize(OverrideTexture->GetSurfaceWidth(), OverrideTexture->GetSurfaceHeight());
+		const float RawZoom = GetCurrentZoom();
+		const float TextureScale = FMath::Min(ViewportSize.X / TextureSize.X, ViewportSize.Y / TextureSize.Y);
+		const FVector2D TextureOffset = 0.5 * (ViewportSize - TextureSize * TextureScale);
+		
+		FCanvasTileItem TileItem(TextureOffset * RawZoom - ViewportSize * RawZoom * ViewRect.Min, OverrideTexture->GetResource(), TextureSize * TextureScale * RawZoom, FLinearColor::White);
+		TileItem.BlendMode = ESimpleElementBlendMode::SE_BLEND_Translucent;
+		Canvas->DrawItem(TileItem);
 		return;
 	}
+	
+	PreDraw(InViewport, Canvas);
+	FEditorViewportClient::Draw(InViewport, Canvas);
+}
 
-	Canvas->Clear(FLinearColor::Black);
+void FSimulcamEditorViewportClient::Draw(const FSceneView* View, FPrimitiveDrawInterface* PDI)
+{
+	FLevelEditorViewportClient::Draw(View, PDI);
+}
 
-	UTexture* Texture = SimulcamViewportWeakPtr.Pin()->GetTexture();
-	if (!Texture)
+void FSimulcamEditorViewportClient::DrawCanvas(FViewport& InViewport, FSceneView& View, FCanvas& Canvas)
+{
+	TSharedPtr<FSimulcamViewportElementsProvider> PinnedElementsProvider = ElementsProvider.Pin();
+	if (!PinnedElementsProvider.IsValid())
 	{
 		return;
 	}
 
-	SimulcamEditorViewportWeakPtr.Pin()->CacheEffectiveTextureSize();
+	const FBox2D ConstrainedRect(FVector2D(View.UnscaledViewRect.Min), FVector2D(View.UnscaledViewRect.Max));
+	const float RawZoom = GetCurrentZoom();
+	
+	UTexture* MediaOverlayTexture = PinnedElementsProvider->GetMediaOverlayTexture();
+	if (MediaOverlayTexture && MediaOverlayTexture->GetResource())
+	{
+		const FLinearColor Color = FLinearColor::White * PinnedElementsProvider->GetMediaOverlayOpacity();
 
-	if (SimulcamEditorViewportWeakPtr.Pin()->GetCustomZoomLevel() <= 0)
+		MediaOverlayTexturePosition = ConstrainedRect.Min * (1.0 - RawZoom) + MediaOverlayTextureOffset * RawZoom - ConstrainedRect.GetSize() * RawZoom * ViewRect.Min;
+		
+		FCanvasTileItem TileItem(MediaOverlayTexturePosition, MediaOverlayTexture->GetResource(), MediaOverlayTextureScaledSize * RawZoom, Color);
+		TileItem.BlendMode = ESimpleElementBlendMode::SE_BLEND_Translucent;
+		Canvas.DrawItem(TileItem);
+	}
+	
+	UMaterialInterface* ToolOverlayMaterial = PinnedElementsProvider->GetToolOverlayMaterial();
+	if (ToolOverlayMaterial && ToolOverlayMaterial->GetRenderProxy())
 	{
-		SimulcamEditorViewportWeakPtr.Pin()->SetCustomZoomLevel(-1);
+		FCanvasTileItem TileItem(MediaOverlayTexturePosition, ToolOverlayMaterial->GetRenderProxy(), MediaOverlayTextureScaledSize * RawZoom);
+		TileItem.BlendMode = ESimpleElementBlendMode::SE_BLEND_Translucent;
+		Canvas.DrawItem(TileItem);
 	}
 
-	// Figure out the size we need
-	FVector2D TextureSize = SimulcamEditorViewportWeakPtr.Pin()->CalculateTextureDimensions();
-
-	constexpr float MipLevel = 0;
-	constexpr float LayerIndex = 0;
-	constexpr float SliceIndex = 0;
-
-	TRefCountPtr<FBatchedElementParameters> BatchedElementParameters;
-	BatchedElementParameters = new FBatchedElementTexture2DPreviewParameters(MipLevel, LayerIndex, SliceIndex, false, false, false, false, false, false);
-
-	if (Texture->GetResource() != nullptr)
+	UMaterialInterface* UserOverlayMaterial = PinnedElementsProvider->GetUserOverlayMaterial();
+	if (UserOverlayMaterial && UserOverlayMaterial->GetRenderProxy())
 	{
-		FCanvasTileItem TileItem(CurrentTexturePosition, Texture->GetResource(), TextureSize, FLinearColor::White);
-		TileItem.BlendMode = ESimpleElementBlendMode::SE_BLEND_Opaque;
-		TileItem.BatchedElementParameters = BatchedElementParameters;
-		Canvas->DrawItem(TileItem);
+		FCanvasTileItem TileItem(MediaOverlayTexturePosition, UserOverlayMaterial->GetRenderProxy(), MediaOverlayTextureScaledSize * RawZoom);
+		TileItem.BlendMode = ESimpleElementBlendMode::SE_BLEND_Translucent;
+		Canvas.DrawItem(TileItem);
 	}
 
+	// Draw black bars to cover up areas of the overlay textures that extend beyond the bounds of the unscaled view rect
+	if (ConstrainedRect.Min.X > 0.0)
+	{
+		FCanvasTileItem LeftBar(FVector2D::ZeroVector, GBlackTexture, FVector2D(ConstrainedRect.Min.X, ViewportSize.Y), FLinearColor::White);
+		Canvas.DrawItem(LeftBar);
+	}
+	
+	if (ConstrainedRect.Max.X < ViewportSize.X)
+	{
+		FCanvasTileItem RightBar(FVector2D(ConstrainedRect.Max.X, 0.0), GBlackTexture, FVector2D(ViewportSize.X - ConstrainedRect.Max.X, ViewportSize.Y), FLinearColor::White);
+		Canvas.DrawItem(RightBar);
+	}
+		
+	if (ConstrainedRect.Min.Y > 0.0)
+	{
+		FCanvasTileItem TopBar(FVector2D::ZeroVector, GBlackTexture, FVector2D(ViewportSize.X, ConstrainedRect.Min.Y), FLinearColor::White);
+		Canvas.DrawItem(TopBar);
+	}
+	
+	if (ConstrainedRect.Max.Y < ViewportSize.Y)
+	{
+		FCanvasTileItem BottomBar(FVector2D(0.0, ConstrainedRect.Max.Y), GBlackTexture, FVector2D(ViewportSize.X, ViewportSize.Y - ConstrainedRect.Max.Y), FLinearColor::White);
+		Canvas.DrawItem(BottomBar);
+	}
+	
 	// If the user is current doing a marquee select, draw the marquee selection box
 	if (bIsMarqueeSelecting)
 	{
 		FCanvasBoxItem BoxItem(SelectionStartCanvas, SelectionBoxSize);
-		Canvas->DrawItem(BoxItem);
+		Canvas.DrawItem(BoxItem);
 	}
 }
 
-void FSimulcamEditorViewportClient::MouseMove(FViewport* Viewport, int32 X, int32 Y)
+void FSimulcamEditorViewportClient::MouseMove(FViewport* InViewport, int32 X, int32 Y)
 {
 	MousePosition = FIntPoint(X, Y);
 }
 
-void FSimulcamEditorViewportClient::ZoomOnPoint(FViewport* Viewport, FIntPoint InPoint, TFunction<void()> ZoomFunction)
-{
-	double CurrentZoom = SimulcamEditorViewportWeakPtr.Pin()->GetCustomZoomLevel();
-
-	ZoomFunction();
-
-	double NewZoom = SimulcamEditorViewportWeakPtr.Pin()->GetCustomZoomLevel();
-	FIntPoint ViewportSize = Viewport->GetSizeXY();
-
-	FVector2D TextureSize = SimulcamEditorViewportWeakPtr.Pin()->CalculateTextureDimensions();
-
-	if (TextureSize.IsNearlyZero())
-	{
-		return;
-	}
-
-	float XPercent = (InPoint.X - CurrentTexturePosition.X) / TextureSize.X;
-	float YPercent = (InPoint.Y - CurrentTexturePosition.Y) / TextureSize.Y;
-
-	CurrentTexturePosition.X = InPoint.X - (XPercent * (TextureSize.X * NewZoom / CurrentZoom));
-	CurrentTexturePosition.Y = InPoint.Y - (YPercent * (TextureSize.Y * NewZoom / CurrentZoom));
-}
-
-void FSimulcamEditorViewportClient::ZoomTowardsFit(FViewport* Viewport)
+bool FSimulcamEditorViewportClient::InputKey(const FInputKeyEventArgs& InEventArgs)
 {
-	double MinZoomLevel = SimulcamEditorViewportWeakPtr.Pin()->GetMinZoomLevel();
-
-	SimulcamEditorViewportWeakPtr.Pin()->ZoomOut();
-
-	double CurrentZoom = SimulcamEditorViewportWeakPtr.Pin()->GetCustomZoomLevel();
-	FIntPoint ViewportSize = Viewport->GetSizeXY();
-
-	FVector2D WantedTexturePosition = SimulcamEditorViewportWeakPtr.Pin()->GetFitPosition();
-
-	if (CurrentZoom >= MinZoomLevel)
+	TSharedPtr<FSimulcamViewportElementsProvider> PinnedElementsProvider = ElementsProvider.Pin();
+	if (!PinnedElementsProvider.IsValid())
 	{
-		double Gradient = 1.0f / ((CurrentZoom - MinZoomLevel) / SSimulcamEditorViewport::ZoomStep);
-		CurrentTexturePosition = FMath::Lerp(CurrentTexturePosition, WantedTexturePosition, FMath::Clamp(Gradient, 0.0, 1.0));
+		return false;
 	}
-}
-
-bool FSimulcamEditorViewportClient::InputKey(const FInputKeyEventArgs& InEventArgs)
-{
+	
 	if (InEventArgs.Event == IE_Pressed)
 	{
-		if (!SimulcamEditorViewportWeakPtr.IsValid())
-		{
-			return false;
-		}
-
 		const bool bIsCtrlDown = FSlateApplication::Get().GetPlatformApplication()->GetModifierKeys().IsControlDown();
 		const bool bIsAltDown = FSlateApplication::Get().GetPlatformApplication()->GetModifierKeys().IsAltDown();
 
@@ -137,14 +171,13 @@ bool FSimulcamEditorViewportClient::InputKey(const FInputKeyEventArgs& InEventAr
 
 		if (InEventArgs.Key == EKeys::LeftMouseButton || InEventArgs.Key == EKeys::MiddleMouseButton || InEventArgs.Key == EKeys::RightMouseButton)
 		{
-			const FGeometry& MyGeometry = SimulcamEditorViewportWeakPtr.Pin()->GetTickSpaceGeometry();
 			// check if we are under the viewport, otherwise the capture system will blindly trigger the PointerEvent
 			if (LocalMouse.ComponentwiseAllGreaterOrEqual(FVector2D(0, 0)) &&
-				LocalMouse.ComponentwiseAllLessThan(MyGeometry.GetAbsoluteSize()))
+				LocalMouse.ComponentwiseAllLessThan(InEventArgs.Viewport->GetSizeXY()))
 			{
 				// create fake geometry and mouseposition
-				const FVector2D FakeMousePosition = GetTexturePosition();
-				const FIntPoint TextureSize = GetTextureSize();
+				const FVector2D FakeMousePosition = GetMousePositionInMediaOverlaySpace();
+				const FIntPoint TextureSize = MediaOverlayTextureRawSize;
 
 				// check for meaningful position
 				if (FakeMousePosition.X >= 0 && FakeMousePosition.Y >= 0 && FakeMousePosition.X < TextureSize.X && FakeMousePosition.Y < TextureSize.Y)
@@ -170,7 +203,7 @@ bool FSimulcamEditorViewportClient::InputKey(const FInputKeyEventArgs& InEventAr
 							0,
 							FSlateApplication::Get().GetPlatformApplication()->GetModifierKeys());
 
-						SimulcamEditorViewportWeakPtr.Pin()->OnViewportClicked(FakeGeometry, PointerEvent);
+						PinnedElementsProvider->OnViewportClicked(FakeGeometry, PointerEvent);
 					}
 				}
 			}
@@ -180,32 +213,24 @@ bool FSimulcamEditorViewportClient::InputKey(const FInputKeyEventArgs& InEventAr
 		{
 			if (InEventArgs.Key == EKeys::MouseScrollUp || (InEventArgs.Key == EKeys::Add && bIsCtrlDown))
 			{
-				ZoomOnPoint(InEventArgs.Viewport, MousePosition, [this] {SimulcamEditorViewportWeakPtr.Pin()->ZoomIn(); });
+				ChangeZoom(InEventArgs.Viewport, MousePosition, ZoomIncrement);
 				return true;
 			}
 
 			if (InEventArgs.Key == EKeys::MouseScrollDown || (InEventArgs.Key == EKeys::Subtract && bIsCtrlDown))
-			{	
-				FVector2D TextureSize = SimulcamEditorViewportWeakPtr.Pin()->CalculateTextureDimensions();
-				if (!bIsCtrlDown && (CurrentTexturePosition.X > 0 || (TextureSize.X + CurrentTexturePosition.X) < InEventArgs.Viewport->GetSizeXY().X || CurrentTexturePosition.Y > 0 || (TextureSize.Y + CurrentTexturePosition.Y) < InEventArgs.Viewport->GetSizeXY().Y))
-				{
-					ZoomTowardsFit(InEventArgs.Viewport);
-				}
-				else
-				{
-					ZoomOnPoint(InEventArgs.Viewport, MousePosition, [this] {SimulcamEditorViewportWeakPtr.Pin()->ZoomOut(); });
-				}
+			{
+				ChangeZoom(InEventArgs.Viewport, MousePosition, -ZoomIncrement);
 				return true;
 			}
 
 			if ((InEventArgs.Key == EKeys::Zero || InEventArgs.Key == EKeys::NumPadZero) && bIsCtrlDown)
 			{
-				ZoomToFit(InEventArgs.Viewport);
+				ResetZoom();
 				return true;
 			}
 		}
 
-		return SimulcamViewportWeakPtr.Pin()->OnViewportInputKey(InEventArgs.Key, InEventArgs.Event);
+		return PinnedElementsProvider->OnViewportInputKey(InEventArgs.Key, InEventArgs.Event);
 	}
 	else if (InEventArgs.Event == IE_Released)
 	{
@@ -215,80 +240,180 @@ bool FSimulcamEditorViewportClient::InputKey(const FInputKeyEventArgs& InEventAr
 			bIsMarqueeSelecting = false;
 
 			// Calculate where the selection ended in texture coordinates based on the latest selection box size
-			const double CurrentZoom = SimulcamEditorViewportWeakPtr.Pin()->GetCustomZoomLevel();
+			const double CurrentZoom = GetCurrentZoom() * MediaOverlayTextureScale;
 			FVector2D SelectionEndTexture = (SelectionStartTexture + (SelectionBoxSize / CurrentZoom));
 
 			// Clamp end position to the texture size
-			const FIntPoint TextureSize = GetTextureSize();
+			const FIntPoint TextureSize = MediaOverlayTextureRawSize;
 			SelectionEndTexture.X = FMath::Clamp(SelectionEndTexture.X, 0, TextureSize.X);
 			SelectionEndTexture.Y = FMath::Clamp(SelectionEndTexture.Y, 0, TextureSize.Y);
 
 			const FVector2D TopLeftPoint = FIntPoint(FMath::Min(SelectionStartTexture.X, SelectionEndTexture.X), FMath::Min(SelectionStartTexture.Y, SelectionEndTexture.Y));
 			const FVector2D BottomRightPoint = FIntPoint(FMath::Max(SelectionStartTexture.X, SelectionEndTexture.X), FMath::Max(SelectionStartTexture.Y, SelectionEndTexture.Y));
 
-			SimulcamViewportWeakPtr.Pin()->OnMarqueeSelect(TopLeftPoint, BottomRightPoint);
+			PinnedElementsProvider->OnMarqueeSelect(TopLeftPoint, BottomRightPoint);
 		}
 
-		return SimulcamViewportWeakPtr.Pin()->OnViewportInputKey(InEventArgs.Key, InEventArgs.Event);
+		return PinnedElementsProvider->OnViewportInputKey(InEventArgs.Key, InEventArgs.Event);
 	}
 	else if (InEventArgs.Event == IE_Repeat)
 	{
-		return SimulcamViewportWeakPtr.Pin()->OnViewportInputKey(InEventArgs.Key, InEventArgs.Event);
+		return PinnedElementsProvider->OnViewportInputKey(InEventArgs.Key, InEventArgs.Event);
 	}
 
 	return false;
 }
 
-void FSimulcamEditorViewportClient::OnViewportResized(FViewport* InViewport, uint32 InParams)
+float FSimulcamEditorViewportClient::GetCurrentZoom() const
+{
+	return 1.0 / GetCurrentInvZoom();
+}
+
+float FSimulcamEditorViewportClient::GetCurrentInvZoom() const
+{
+	return ViewRect.Max.X - ViewRect.Min.X;
+}
+
+float FSimulcamEditorViewportClient::GetMediaOverlayScale() const
+{
+	return GetCurrentZoom() * MediaOverlayTextureScale;
+}
+
+#define SAFE_DIVIDE(Dividend, Divisor, Default) (Divisor != 0.0f) ? (Dividend / Divisor) : Default
+
+void FSimulcamEditorViewportClient::PreDraw(FViewport* InViewport, FCanvas* Canvas)
 {
-	if (!SimulcamEditorViewportWeakPtr.IsValid())
+	TSharedPtr<FSimulcamViewportElementsProvider> PinnedElementsProvider = ElementsProvider.Pin();
+	if (!PinnedElementsProvider.IsValid())
 	{
 		return;
 	}
 
-	if (SimulcamEditorViewportWeakPtr.Pin()->GetViewport()->GetViewport() == InViewport)
+	bool bViewportSizeChanged = false;
+	if (ViewportSize != InViewport->GetSizeXY())
 	{
-		FIntPoint NewViewportSize = InViewport->GetSizeXY();
-		if (NewViewportSize != CurrentViewportSize)
+		ViewportSize = InViewport->GetSizeXY();
+		bViewportSizeChanged = true;
+	}
+	
+	if (UTexture* MediaOverlayTexture = PinnedElementsProvider->GetMediaOverlayTexture())
+	{
+		const FIntPoint TextureSize = FIntPoint(MediaOverlayTexture->GetSurfaceWidth(), MediaOverlayTexture->GetSurfaceHeight());
+		bool bOverlayTextureSizeChanged = false;
+		if (TextureSize != MediaOverlayTextureRawSize)
 		{
-			ZoomToFit(InViewport);
-			CurrentViewportSize = NewViewportSize;
+			MediaOverlayTextureRawSize = TextureSize;
+			bOverlayTextureSizeChanged = true;
 		}
-	}
-}
 
-void FSimulcamEditorViewportClient::OnTextureResized()
-{
-	if (TSharedPtr<SSimulcamEditorViewport> SimulcamEditorViewport = SimulcamEditorViewportWeakPtr.Pin())
-	{
-		if (FViewport* Viewport = SimulcamEditorViewport->GetViewport()->GetViewport())
+		if (bViewportSizeChanged || bOverlayTextureSizeChanged)
 		{
-			ZoomToFit(Viewport);
+			MediaOverlayTextureScale = FMath::Min(SAFE_DIVIDE(ViewportSize.X, MediaOverlayTextureRawSize.X, 1.0f), SAFE_DIVIDE(ViewportSize.Y, MediaOverlayTextureRawSize.Y, 1.0f));
+			MediaOverlayTextureScaledSize = FVector2D(MediaOverlayTextureRawSize) * MediaOverlayTextureScale;
+			MediaOverlayTextureOffset = 0.5 * (ViewportSize - MediaOverlayTextureScaledSize);
+
+			ResetZoom();
 		}
 	}
-}
+	else
+	{
+		MediaOverlayTextureRawSize = FIntPoint::ZeroValue;
+		MediaOverlayTextureScale = 1.0f;
+		MediaOverlayTextureScaledSize = FVector2D::ZeroVector;
+		MediaOverlayTextureOffset = FVector2D::ZeroVector;
+	}
+	
+	const float ViewportAspectRatio = SAFE_DIVIDE(ViewportSize.X, ViewportSize.Y, 1.0f);
+	const float CGAspectRatio = bUseControllingActorViewInfo ? ControllingActorViewInfo.AspectRatio : ViewportAspectRatio;
+	const float FilmbackAspectRatio = PinnedElementsProvider->GetCameraFeedAspectRatio();
+	const float MediaOverlayAspectRatio = SAFE_DIVIDE(MediaOverlayTextureScaledSize.X, MediaOverlayTextureScaledSize.Y, 1.0f);
+
+	// Compute the rectangle that the CG layer is intended to inhabit within the full viewport, as it will be aspect ratio constrained inside the media overlay texture
+	CGViewportSize = CGAspectRatio > MediaOverlayAspectRatio ?
+		FVector2D(MediaOverlayTextureScaledSize.X, SAFE_DIVIDE(MediaOverlayTextureScaledSize.X, CGAspectRatio, 0.0f)) :
+		FVector2D(MediaOverlayTextureScaledSize.Y * CGAspectRatio, MediaOverlayTextureScaledSize.Y);
+
+	// Need to modify the CG viewport size to match the aspect ratio of the viewport being rendered, as the viewport size is what determines the backbuffer
+	// size of the render
+	if (ViewportAspectRatio > CGAspectRatio)
+	{
+		CGViewportSize.X = CGViewportSize.Y * ViewportAspectRatio;
+	}
+	else
+	{
+		CGViewportSize.Y = SAFE_DIVIDE(CGViewportSize.X, ViewportAspectRatio, 0.0f);
+	}
 
-void FSimulcamEditorViewportClient::ZoomToFit(FViewport* InViewport)
-{
-	if (SimulcamEditorViewportWeakPtr.IsValid())
+	if (FilmbackAspectRatio < CGAspectRatio)
+	{
+		CGViewportSize.X *= SAFE_DIVIDE(FilmbackAspectRatio, CGAspectRatio, 0.0f);
+	}
+	else
+	{
+		CGViewportSize.Y *= SAFE_DIVIDE(CGAspectRatio, FilmbackAspectRatio, 0.0f);
+	}
+	
+	// TODO: Ideally we would want to unconstrain the aspect ratio so that we can make full use of the viewport's space. However, any change to the aspect ratio
+	// here will change how lens distortion gets applied to the CG image, as the UV coordinates will change as the aspect ratio changes. In theory, this could be
+	// mitigated by something like asymmetric overscan, but the lens distortion scene view extension has not been updated to account for asymmetric overscan yet.
+	//ControllingActorViewInfo.bConstrainAspectRatio = false;
+	//ControllingActorViewInfo.AspectRatioAxisConstraint = ViewportAspectRatio > CGAspectRatio ? EAspectRatioAxisConstraint::AspectRatio_MaintainYFOV : EAspectRatioAxisConstraint::AspectRatio_MaintainXFOV;
+	
+	if (CropToViewRectMaterialInstance.IsValid())
 	{
-		SimulcamEditorViewportWeakPtr.Pin()->SetCustomZoomLevel(-1);
-		FVector2D TextureSize = SimulcamEditorViewportWeakPtr.Pin()->CalculateTextureDimensions();
-		if (!TextureSize.IsNearlyZero())
+		// If the CG viewport size is zero, assume that either the media texture size is zero/invalid or an aspect ratio is invalid. Either way, skip
+		// setting the view rect and bounds in the crop material so that a normal viewport is rendered
+		if (CGViewportSize.IsNearlyZero())
 		{
-			CurrentTexturePosition = FVector2D((InViewport->GetSizeXY().X - TextureSize.X) / 2.0f, (InViewport->GetSizeXY().Y - TextureSize.Y) / 2.0f);
+			CropToViewRectMaterialInstance->ClearParameterValues();
+			return;
 		}
+		
+		float Overscan = 1.0f;
+		if (ElementsProvider.IsValid())
+		{
+			if (ACameraActor* Camera = ElementsProvider.Pin()->GetCameraActor())
+			{
+				Overscan = Camera->GetCameraComponent()->Overscan + 1.0f;
+			}
+		}
+
+		// Compute the view rect from the CG viewport size and the current zoom/pan
+		const FBox2D CGViewRect((ViewRect.Min - 0.5) * ViewportSize / CGViewportSize + 0.5, (ViewRect.Max - 0.5) * ViewportSize / CGViewportSize + 0.5);
+		
+		// In order to crop and scale a distorted camera render that includes overscan, we must convert the view rect to overscanned coordinates,
+		// and also add a bit of overscan since the render pipeline will crop the overscan as the very last thing after post process materials are processed.
+		// For posterity, the equation below is derived from combining the overscan scaling equation (rect = (rect - center) * overscan + center) with
+		// the equation to transform to overscanned coordinate space (rect = (rect - 0.5) / overscan + 0.5)
+		const FBox2D OverscannedViewRect = CGViewRect.ShiftBy((CGViewRect.GetCenter() - 0.5) * (1.0 / Overscan - 1.0));
+		
+		// Similarly to above, if the CG viewport is being inset at a smaller size than the whole viewport, we need to manually crop out the overscan as it will be too inset
+		// for the crop post process pass to crop it out
+		const FBox2D OverscanBounds(-FVector2D(0.5, 0.5) / Overscan + 0.5, FVector2D(0.5, 0.5) / Overscan + 0.5);
+		
+		CropToViewRectMaterialInstance->SetVectorParameterValue("ViewRect", FVector4(OverscannedViewRect.Min.X, OverscannedViewRect.Min.Y, OverscannedViewRect.Max.X, OverscannedViewRect.Max.Y));
+		CropToViewRectMaterialInstance->SetVectorParameterValue("ViewBounds", FVector4(OverscanBounds.Min.X, OverscanBounds.Min.Y, OverscanBounds.Max.X, OverscanBounds.Max.Y));
 	}
 }
 
-bool FSimulcamEditorViewportClient::InputChar(FViewport* Viewport, int32 ControllerId, TCHAR Character)
+FVector2D FSimulcamEditorViewportClient::GetMousePositionInMediaOverlaySpace() const
 {
-	if (!bWithZoom)
-	{
-		return false;
-	}
+	return (FVector2D(MousePosition) - MediaOverlayTexturePosition) * GetCurrentInvZoom() / MediaOverlayTextureScale;
+}
+
+FVector2D FSimulcamEditorViewportClient::GetScaledMediaOverlayTextureSize() const
+{
+	return MediaOverlayTextureScaledSize * GetCurrentZoom();
+}
 
-	if (!SimulcamViewportWeakPtr.IsValid())
+void FSimulcamEditorViewportClient::ResetZoom()
+{
+	ViewRect = FBox2D(FVector2D(0.0, 0.0), FVector2D(1.0, 1.0));
+}
+
+bool FSimulcamEditorViewportClient::InputChar(FViewport* InViewport, int32 ControllerId, TCHAR Character)
+{
+	if (!bWithZoom)
 	{
 		return false;
 	}
@@ -297,13 +422,13 @@ bool FSimulcamEditorViewportClient::InputChar(FViewport* Viewport, int32 Control
 
 	if (Character == TCHAR('+') && bIsCtrlDown)
 	{
-		ZoomOnPoint(Viewport, MousePosition, [this] {SimulcamEditorViewportWeakPtr.Pin()->ZoomIn(); });
+		ChangeZoom(InViewport, MousePosition, ZoomIncrement);
 		return true;
 	}
 
 	if (Character == TCHAR('-') && bIsCtrlDown)
 	{
-		ZoomOnPoint(Viewport, MousePosition, [this] {SimulcamEditorViewportWeakPtr.Pin()->ZoomOut(); });
+		ChangeZoom(InViewport, MousePosition, -ZoomIncrement);
 		return true;
 	}
 
@@ -313,35 +438,27 @@ bool FSimulcamEditorViewportClient::InputChar(FViewport* Viewport, int32 Control
 bool FSimulcamEditorViewportClient::InputAxis(const FInputKeyEventArgs& InEventArgs)
 {
 	const FKey& Key = InEventArgs.Key;
-	const float Delta = InEventArgs.AmountDepressed;
-	FViewport* Viewport = InEventArgs.Viewport;
+	const float MouseDelta = InEventArgs.AmountDepressed;
+	FViewport* InViewport = InEventArgs.Viewport;
 	
 	if (Key == EKeys::MouseX || Key == EKeys::MouseY)
 	{
-		if (ShouldUseMousePanning(Viewport))
+		if (ShouldUseMousePanning(InViewport))
 		{
-			if (Key == EKeys::MouseY && CanPanVertically(Viewport, Delta))
+			if (Key == EKeys::MouseY)
 			{
-				CurrentTexturePosition.Y -= Delta;
-
-				FVector2D TextureSize = SimulcamEditorViewportWeakPtr.Pin()->CalculateTextureDimensions();
-
-				CurrentTexturePosition.Y = FMath::Clamp(CurrentTexturePosition.Y, Viewport->GetSizeXY().Y - TextureSize.Y, 0.0f);
+				ViewRect = ViewRect.ShiftBy(FVector2D(0.0, GetClampedVerticalPanDelta(MouseDelta)));
 			}
-			else if (Key == EKeys::MouseX && CanPanHorizontally(Viewport, Delta))
+			else if (Key == EKeys::MouseX)
 			{
-				CurrentTexturePosition.X += Delta;
-
-				FVector2D TextureSize = SimulcamEditorViewportWeakPtr.Pin()->CalculateTextureDimensions();
-
-				CurrentTexturePosition.X = FMath::Clamp(CurrentTexturePosition.X, Viewport->GetSizeXY().X - TextureSize.X, 0.0f);
+				ViewRect = ViewRect.ShiftBy(FVector2D(GetClampedHorizontalPanDelta(MouseDelta), 0.0));
 			}
 		}
 
 		// Update the marquee selection box size based on the current mouse position
 		if (bIsMarqueeSelecting)
 		{
-			FVector2D LocalMouse = FVector2D(Viewport->GetMouseX(), Viewport->GetMouseY());
+			FVector2D LocalMouse = FVector2D(InViewport->GetMouseX(), InViewport->GetMouseY());
 			SelectionBoxSize = LocalMouse - SelectionStartCanvas;
 			ClampSelectionBoxSizeToTextureSize();
 		}
@@ -354,8 +471,8 @@ bool FSimulcamEditorViewportClient::InputAxis(const FInputKeyEventArgs& InEventA
 
 void FSimulcamEditorViewportClient::ClampSelectionBoxSizeToTextureSize()
 {
-	const FIntPoint TextureSize = GetTextureSize();
-	const double CurrentZoom = SimulcamEditorViewportWeakPtr.Pin()->GetCustomZoomLevel();
+	const FIntPoint TextureSize = MediaOverlayTextureRawSize;
+	const double CurrentZoom = GetCurrentZoom() * MediaOverlayTextureScale;
 
 	const FVector2D MinBoxSize = (FVector2D(0.0f, 0.0f) - SelectionStartTexture) * CurrentZoom;
 	const FVector2D MaxBoxSize = (FVector2D(TextureSize) - SelectionStartTexture) * CurrentZoom;
@@ -364,108 +481,144 @@ void FSimulcamEditorViewportClient::ClampSelectionBoxSizeToTextureSize()
 	SelectionBoxSize.Y = FMath::Clamp(SelectionBoxSize.Y, MinBoxSize.Y, MaxBoxSize.Y);
 }
 
-bool FSimulcamEditorViewportClient::ShouldUseMousePanning(FViewport* Viewport) const
+bool FSimulcamEditorViewportClient::ShouldUseMousePanning(FViewport* InViewport) const
 {
-	return bWithPan && SimulcamViewportWeakPtr.IsValid() && SimulcamViewportWeakPtr.Pin()->GetTexture() && Viewport->KeyState(EKeys::RightMouseButton);
+	return bWithPan && InViewport->KeyState(EKeys::RightMouseButton);
 }
 
-bool FSimulcamEditorViewportClient::CanPanHorizontally(FViewport* Viewport, float Direction) const
+float FSimulcamEditorViewportClient::GetClampedHorizontalPanDelta(float MouseDelta)
 {
-	FVector2D TextureSize = SimulcamEditorViewportWeakPtr.Pin()->CalculateTextureDimensions();
-	if (Direction < 0)
+	const float ViewportWidth = Viewport->GetSizeXY().X;
+	const float ScaledMouseDelta = MouseDelta / ViewportWidth;
+
+	// Mouse moved left, which means a pan to the right
+	if (MouseDelta < 0)
+	{
+		return FMath::Min(-ScaledMouseDelta, 1.0 - ViewRect.Max.X);
+	}
+
+	// Otherwise, mouse moved right, which means pan left
+	return FMath::Max(-ScaledMouseDelta, -ViewRect.Min.X);
+}
+
+float FSimulcamEditorViewportClient::GetClampedVerticalPanDelta(float MouseDelta)
+{
+	const float ViewportHeight = Viewport->GetSizeXY().Y;
+	const float ScaledMouseDelta = MouseDelta / ViewportHeight;
+
+	// Mouse moved up, which means a pan down
+	if (MouseDelta > 0)
 	{
-		return (TextureSize.X + CurrentTexturePosition.X) > Viewport->GetSizeXY().X;
+		return FMath::Min(ScaledMouseDelta, 1.0 - ViewRect.Max.Y);
 	}
-	return CurrentTexturePosition.X < 0;
+
+	// Otherwise, mouse moved down, which means pan up
+	return FMath::Max(ScaledMouseDelta, -ViewRect.Min.Y);
 }
 
-bool FSimulcamEditorViewportClient::CanPanVertically(FViewport* Viewport, float Direction) const
+void FSimulcamEditorViewportClient::ChangeZoom(FViewport* InViewport, const FIntPoint& InMousePosition, float ZoomDelta)
 {
-	FVector2D TextureSize = SimulcamEditorViewportWeakPtr.Pin()->CalculateTextureDimensions();
-	if (Direction > 0)
+	const float CurrentZoom = GetCurrentZoom();
+	const float NewZoom = FMath::Clamp(CurrentZoom + ZoomDelta, 1.0, MaxZoom);
+	
+	ViewRect.Min = ViewRect.Min + FVector2D(InMousePosition) / FVector2D(InViewport->GetSizeXY()) * ( 1.0 / CurrentZoom - 1.0 / NewZoom);
+	ViewRect.Max = ViewRect.Min + 1.0 / NewZoom;
+
+	// Correct zoom rect back towards the 0-1 range in case the zoom change caused an overshoot
+	FVector2D Correction = FVector2D::ZeroVector;
+	if (ViewRect.Max.X > 1.0)
+	{
+		Correction.X = 1.0 - ViewRect.Max.X;
+	}
+	else if (ViewRect.Min.X < 0.0)
+	{
+		Correction.X = -ViewRect.Min.X;
+	}
+
+	if (ViewRect.Max.Y > 1.0)
+	{
+		Correction.Y = 1.0 - ViewRect.Max.Y;
+	}
+	else if (ViewRect.Min.Y < 0.0)
 	{
-		return (TextureSize.Y + CurrentTexturePosition.Y) > Viewport->GetSizeXY().Y;
+		Correction.Y = -ViewRect.Min.Y;
 	}
-	return CurrentTexturePosition.Y < 0;
+
+	ViewRect = ViewRect.ShiftBy(Correction);
 }
 
-EMouseCursor::Type FSimulcamEditorViewportClient::GetCursor(FViewport* Viewport, int32 X, int32 Y)
+EMouseCursor::Type FSimulcamEditorViewportClient::GetCursor(FViewport* InViewport, int32 X, int32 Y)
 {
-	return ShouldUseMousePanning(Viewport) ? EMouseCursor::GrabHandClosed : EMouseCursor::Default;
+	return ShouldUseMousePanning(InViewport) ? EMouseCursor::GrabHandClosed : EMouseCursor::Default;
 }
 
 FText FSimulcamEditorViewportClient::GetDisplayedResolution() const
 {
-	if (!SimulcamEditorViewportWeakPtr.IsValid() || !SimulcamViewportWeakPtr.IsValid() || !SimulcamViewportWeakPtr.Pin()->HasValidTextureResource())
+	TSharedPtr<FSimulcamViewportElementsProvider> PinnedElementsProvider = ElementsProvider.Pin();
+	if (!PinnedElementsProvider.IsValid())
 	{
 		return FText::GetEmpty();
 	}
 
-	FVector2D TextureSize = SimulcamEditorViewportWeakPtr.Pin()->CalculateTextureDimensions();
+	UTexture* Texture = PinnedElementsProvider->GetOverrideTexture();
+	const bool bIsOverrideTexture = Texture != nullptr;
+	
+	if (!bIsOverrideTexture)
+	{
+		Texture = PinnedElementsProvider->GetMediaOverlayTexture();
+	}
 
-	const FVector2D TexturePosition = GetTexturePosition();
-	const double CurrentZoom = SimulcamEditorViewportWeakPtr.Pin()->GetCustomZoomLevel();
+	if (!Texture)
+	{
+		return FText::GetEmpty();
+	}
 
-	uint32 TextureHeight = 1;
-	uint32 TextureWidth = 1;
-	UTexture* Texture = SimulcamViewportWeakPtr.Pin()->GetTexture();
+	FVector2D TextureSize;
+	FVector2D ScaledTextureSize;
+	float Zoom;
+	FVector2D MouseTexturePosition;
 
-	if (Texture && Texture->GetResource())
+	if (bIsOverrideTexture)
+	{
+		TextureSize = FVector2D(Texture->GetSurfaceWidth(), Texture->GetSurfaceHeight());
+		Zoom = GetCurrentZoom() * FMath::Min(ViewportSize.X / TextureSize.X, ViewportSize.Y / TextureSize.Y);
+		ScaledTextureSize = TextureSize * Zoom;
+		MouseTexturePosition = (FVector2D(MousePosition) - 0.5 * (ViewportSize - ScaledTextureSize * GetCurrentInvZoom())) / Zoom;
+	}
+	else
 	{
-		TextureWidth = Texture->GetResource()->GetSizeX();
-		TextureHeight = Texture->GetResource()->GetSizeY();
+		TextureSize = MediaOverlayTextureRawSize;
+		Zoom = GetCurrentZoom() * MediaOverlayTextureScale;
+		ScaledTextureSize = GetScaledMediaOverlayTextureSize();
+		MouseTexturePosition = GetMousePositionInMediaOverlaySpace();
 	}
 
-	int32 TexturePositionX = TexturePosition.X;
-	if (TexturePositionX < 0)
+	if (MouseTexturePosition.X < 0)
 	{
-		TexturePositionX = 0;
+		MouseTexturePosition.X = 0.0;
 	}
-	else if (TexturePositionX > static_cast<int32>(TextureWidth))
+	else if (MouseTexturePosition.X > TextureSize.X)
 	{
-		TexturePositionX = TextureWidth;
+		MouseTexturePosition.X = TextureSize.X;
 	}
 
-	int32 TexturePositionY = TexturePosition.Y;
-	if (TexturePositionY < 0)
+	if (MouseTexturePosition.Y < 0)
 	{
-		TexturePositionY = 0;
+		MouseTexturePosition.Y = 0.0;
 	}
-	else if (TexturePositionY > static_cast<int32>(TextureHeight))
+	else if (MouseTexturePosition.Y > TextureSize.Y)
 	{
-		TexturePositionY = TextureHeight;
+		MouseTexturePosition.Y = TextureSize.Y;
 	}
 
 	return FText::Format(
 		FText::FromString("Displayed: {0}x{1}\nTextureSize: {2}x{3}\nTexturePosition: {4}x{5}\nZoom: {6}"),
-		FText::AsNumber(FMath::Max(1.0f, TextureSize.X)),
-		FText::AsNumber(FMath::Max(1.0f, TextureSize.Y)),
-		FText::AsNumber(TextureWidth),
-		FText::AsNumber(TextureHeight),
-		FText::AsNumber(TexturePositionX),
-		FText::AsNumber(TexturePositionY),
-		FText::AsNumber(CurrentZoom)
+		FText::AsNumber(ScaledTextureSize.X),
+		FText::AsNumber(ScaledTextureSize.Y),
+		FText::AsNumber(TextureSize.X),
+		FText::AsNumber(TextureSize.Y),
+		FText::AsNumber((int32)MouseTexturePosition.X),
+		FText::AsNumber((int32)MouseTexturePosition.Y),
+		FText::AsNumber(Zoom)
 	);
 }
-
-FVector2D FSimulcamEditorViewportClient::GetTexturePosition() const
-{
-	const double CurrentZoom = SimulcamEditorViewportWeakPtr.Pin()->GetCustomZoomLevel();
-	return FVector2D((MousePosition.X - CurrentTexturePosition.X) / CurrentZoom, (MousePosition.Y - CurrentTexturePosition.Y) / CurrentZoom);
-}
-
-FIntPoint FSimulcamEditorViewportClient::GetTextureSize() const 
-{
-	FIntPoint TextureSize = FIntPoint(1);
-
-	if (UTexture* Texture = SimulcamViewportWeakPtr.Pin()->GetTexture())
-	{
-		if (Texture->GetResource())
-		{
-			TextureSize.X = Texture->GetResource()->GetSizeX();
-			TextureSize.Y = Texture->GetResource()->GetSizeY();
-		}
-	}
-
-	return TextureSize;
-}
diff --git a/Engine/Plugins/VirtualProduction/CameraCalibration/Source/CameraCalibrationEditor/Private/AssetEditor/Viewport/SimulcamEditorViewportClient.h b/Engine/Plugins/VirtualProduction/CameraCalibration/Source/CameraCalibrationEditor/Private/AssetEditor/Viewport/SimulcamEditorViewportClient.h
--- a/Engine/Plugins/VirtualProduction/CameraCalibration/Source/CameraCalibrationEditor/Private/AssetEditor/Viewport/SimulcamEditorViewportClient.h
+++ b/Engine/Plugins/VirtualProduction/CameraCalibration/Source/CameraCalibrationEditor/Private/AssetEditor/Viewport/SimulcamEditorViewportClient.h
@@ -1,11 +1,9 @@
 // Copyright Epic Games, Inc. All Rights Reserved.
 #pragma once
 
-#include "UnrealClient.h"
 #include "CoreMinimal.h"
 #include "InputCoreTypes.h"
-#include "UObject/GCObject.h"
-#include "ViewportClient.h"
+#include "LevelEditorViewport.h"
 
 class FCanvas;
 class ITextureEditorToolkit;
@@ -14,62 +12,118 @@ class UTexture2D;
 class SSimulcamViewport;
 class SSimulcamEditorViewport;
 
+/** Interface used by the simulcam viewport client to retrieve any elements used for rendering the simulcam viewport */
+class FSimulcamViewportElementsProvider
+{
+public:
+	virtual ACameraActor* GetCameraActor() const { return nullptr; }
+	virtual float GetCameraFeedAspectRatio() const { return 1.0f; }
+	virtual UTexture* GetMediaOverlayTexture() const { return nullptr; }
+	virtual float GetMediaOverlayOpacity() const { return 1.0f; }
+	virtual UMaterialInterface* GetToolOverlayMaterial() const { return nullptr;}
+	virtual UMaterialInterface* GetUserOverlayMaterial() const { return nullptr;}
+	virtual UTexture* GetOverrideTexture() const { return nullptr; }
+
+	virtual void OnViewportClicked(const FGeometry& InGeometry, const FPointerEvent& InPointerEvent) { }
+	virtual bool OnViewportInputKey(const FKey& InKey, const EInputEvent& InInputEvent) { return false; }
+	virtual void OnMarqueeSelect(const FVector2D& InStartPosition, const FVector2D& InEndPosition) { }
+};
+
 class FSimulcamEditorViewportClient
-	: public FViewportClient
+	: public FLevelEditorViewportClient
 {
+public:
+	static constexpr double MaxZoom = 16.0;
+	static constexpr double ZoomIncrement = 0.1;
+	
 public:
 	/** Constructor */
-	FSimulcamEditorViewportClient(const TSharedRef<SSimulcamViewport>& InSimulcamViewport, const TSharedRef<SSimulcamEditorViewport>& InSimulcamEditorViewport, const bool bInWithZoom, const bool bInWithPan);
+	FSimulcamEditorViewportClient(TSharedPtr<FSimulcamViewportElementsProvider> InElementsProvider, const bool bInWithZoom, const bool bInWithPan);
 
 	/** Begin FViewportClient interface */
-	virtual void Draw(FViewport* Viewport, FCanvas* Canvas) override;
+	virtual void Tick(float DeltaSeconds) override;
+	virtual void Draw(FViewport* InViewport, FCanvas* Canvas) override;
+	virtual void Draw(const FSceneView* View, FPrimitiveDrawInterface* PDI) override;
+	virtual void DrawCanvas(FViewport& InViewport, FSceneView& View, FCanvas& Canvas) override;
 	virtual bool InputKey(const FInputKeyEventArgs& InEventArgs) override;
-	virtual bool InputChar(FViewport* Viewport, int32 ControllerId, TCHAR Character) override;
+	virtual bool InputChar(FViewport* InViewport, int32 ControllerId, TCHAR Character) override;
 	virtual bool InputAxis(const FInputKeyEventArgs& InEventArgs) override;
-	virtual UWorld* GetWorld() const override { return nullptr; }
-	virtual EMouseCursor::Type GetCursor(FViewport* Viewport, int32 X, int32 Y) override;
-	virtual void MouseMove(FViewport* Viewport, int32 X, int32 Y) override;
+	virtual EMouseCursor::Type GetCursor(FViewport* InViewport, int32 X, int32 Y) override;
+	virtual void MouseMove(FViewport* InViewport, int32 X, int32 Y) override;
 	/** End FViewportClient interface */
 
 	/** Returns a string representation of the currently displayed textures resolution */
 	FText GetDisplayedResolution() const;
 
-	/** Triggered whenever the viewport is resized */
-	void OnViewportResized(FViewport* InViewport, uint32 InParams);
+	/** Gets the current zoom level of the view rectangle */
+	float GetCurrentZoom() const;
 
-	/** Triggered whenever the underlying texture resource is resized */
-	void OnTextureResized();
+	/** Gets the current inverse zoom level of the view rectangle */
+	float GetCurrentInvZoom() const;
 
-private:
-	FVector2D GetTexturePosition() const;
-	FIntPoint GetTextureSize() const;
+	/** Gets the amount that the media overlay texture is being scaled by, including zoom scaling */
+	float GetMediaOverlayScale() const;
 
+private:
+	/** Performs any necessary processing needed before the viewport client renders */
+	void PreDraw(FViewport* InViewport, FCanvas* Canvas);
+	
+	/** Gets the current mouse position in the pixel coordinates of the media overlay texture */
+	FVector2D GetMousePositionInMediaOverlaySpace() const;
+
+	/** Gets the size the media overlay texture is being scaled to when displayed in the viewport */
+	FVector2D GetScaledMediaOverlayTextureSize() const;
+	
 	/** TRUE if right clicking and dragging for panning a texture 2D */
-	bool ShouldUseMousePanning(FViewport* Viewport) const;
+	bool ShouldUseMousePanning(FViewport* InViewport) const;
 
-	bool CanPanHorizontally(FViewport* Viewport, float Direction) const;
-	bool CanPanVertically(FViewport* Viewport, float Direction) const;
+	/** Converts a horizontal mouse delta in pixels and converts it to relative coordinates, clamping to ensure that the view rect stays bounded between 0.0 and 1.0 */
+	float GetClampedHorizontalPanDelta(float MouseDelta);
 
-	/** Zoom under a specific pixel maintaining it aligned to the cursor */
-	void ZoomOnPoint(FViewport* Viewport, FIntPoint InPoint, TFunction<void()> ZoomFunction);
+	/** Converts a vertical mouse delta in pixels and converts it to relative coordinates, clamping to ensure that the view rect stays bounded between 0.0 and 1.0 */
+	float GetClampedVerticalPanDelta(float MouseDelta);
 
-	/** Zoom to fit the texture in the viewport */
-	void ZoomToFit(FViewport* InViewport);
+	void ChangeZoom(FViewport* InViewport, const FIntPoint& InMousePosition, float ZoomDelta);
 
-	/** Zoom Out trying to (smoothly) re-fit the image if there is space on borders */
-	void ZoomTowardsFit(FViewport* Viewport);
+	/** Zoom to fit the texture in the viewport */
+	void ResetZoom();
 
 	/** Clamp the selection box size to only be large as the texture size */
 	void ClampSelectionBoxSizeToTextureSize();
 
-	TWeakPtr<SSimulcamViewport> SimulcamViewportWeakPtr;
-	TWeakPtr<SSimulcamEditorViewport> SimulcamEditorViewportWeakPtr;
-
-	FIntPoint MousePosition;
-
-	FVector2D CurrentTexturePosition = FVector2D(0, 0);
-
-	FIntPoint CurrentViewportSize = FIntPoint(0, 0);
+private:
+	/** An interface that supplies all the necessary assets needed to render the simulcam viewport, and has handlers for various input events */
+	TWeakPtr<FSimulcamViewportElementsProvider> ElementsProvider;
+
+	/** Post process material to crop the viewport's output to the zoom/pan view rect */
+	TStrongObjectPtr<UMaterialInstanceDynamic> CropToViewRectMaterialInstance;
+
+	/** The position of the mouse within the viewport */
+	FIntPoint MousePosition = FIntPoint::ZeroValue;
+
+	/** Cached viewport size, to keep track of when the viewport size has changed so various values can be recomputed */
+	FVector2D ViewportSize = FVector2D::ZeroVector;
+
+	/** The size of the CG viewport when it is scaled down to fit within the media overlay */
+	FVector2D CGViewportSize = FVector2D::ZeroVector;
+	
+	/** Box that represents the portion of the texture and viewport the user is zoomed and panned to, in relative coordinates ((0, 0) is top left, (1, 1) is bottom right) */
+	FBox2D ViewRect = FBox2D(FVector2D(0.0, 0.0), FVector2D(1.0, 1.0));
+
+	/** The offset from the viewport's top left corner that the media overlay texture must be placed to center it in the viewport */
+	FVector2D MediaOverlayTextureOffset = FVector2D::ZeroVector;
+	
+	/** The position in pixel coordinates of the viewport at which the media overlay texture was last drawn */
+	FVector2D MediaOverlayTexturePosition = FVector2D::ZeroVector;
+
+	/** The raw size of the media overlay texture */
+	FIntPoint MediaOverlayTextureRawSize = FIntPoint::ZeroValue;
+
+	/** The size of the media overlay texture when scaled to fit inside the viewport */
+	FVector2D MediaOverlayTextureScaledSize = FVector2D::ZeroVector;
+	
+	/** The scale applied to the media overlay texture to match the rendered viewport size when the media overlay texture was last drawn */
+	double MediaOverlayTextureScale = 1.0;
 
 	bool bWithZoom;
 	bool bWithPan;
diff --git a/Engine/Plugins/VirtualProduction/CameraCalibration/Source/CameraCalibrationEditor/Private/Calibrators/CameraImageCenterAlgoManual.cpp b/Engine/Plugins/VirtualProduction/CameraCalibration/Source/CameraCalibrationEditor/Private/Calibrators/CameraImageCenterAlgoManual.cpp
--- a/Engine/Plugins/VirtualProduction/CameraCalibration/Source/CameraCalibrationEditor/Private/Calibrators/CameraImageCenterAlgoManual.cpp
+++ b/Engine/Plugins/VirtualProduction/CameraCalibration/Source/CameraCalibrationEditor/Private/Calibrators/CameraImageCenterAlgoManual.cpp
@@ -53,14 +53,7 @@ void UCameraImageCenterAlgoManual::Activate()
 		return;
 	}
 
-	// Cache the render target size to use when calculating the fine adjustments
-	UTextureRenderTarget2D* RenderTarget = StepsController->GetRenderTarget();
-	if (!RenderTarget)
-	{
-		return;
-	}
-
-	RenderTargetSize = FIntPoint(RenderTarget->SizeX, RenderTarget->SizeY);
+	RenderTargetSize = StepsController->GetMediaOverlayResolution();
 
 	const FLensFileEvaluationInputs EvalInputs = StepsController->GetLensFileEvaluationInputs();
 	if (!EvalInputs.bIsValid)
diff --git a/Engine/Plugins/VirtualProduction/CameraCalibration/Source/CameraCalibrationEditor/Private/Calibrators/CameraNodalOffsetAlgoOpticalAxis.cpp b/Engine/Plugins/VirtualProduction/CameraCalibration/Source/CameraCalibrationEditor/Private/Calibrators/CameraNodalOffsetAlgoOpticalAxis.cpp
--- a/Engine/Plugins/VirtualProduction/CameraCalibration/Source/CameraCalibrationEditor/Private/Calibrators/CameraNodalOffsetAlgoOpticalAxis.cpp
+++ b/Engine/Plugins/VirtualProduction/CameraCalibration/Source/CameraCalibrationEditor/Private/Calibrators/CameraNodalOffsetAlgoOpticalAxis.cpp
@@ -37,7 +37,12 @@ void UCameraNodalOffsetAlgoOpticalAxis::Initialize(UNodalOffsetTool* InNodalOffs
 	}
 
 	// Create the render target for the crosshair overlay
-	const FIntPoint OverlaySize = StepsController->GetCompRenderResolution();
+	FIntPoint OverlaySize = StepsController->GetMediaOverlayResolution();
+	if (OverlaySize.SizeSquared() == 0)
+	{
+		// If the media overlay does not have a valid resolution yet, use a default size
+		OverlaySize = FIntPoint(1920, 1080);
+	}
 
 	OverlayTexture = NewObject<UTextureRenderTarget2D>(GetTransientPackage(), MakeUniqueObjectName(GetTransientPackage(), UTextureRenderTarget2D::StaticClass()));
 
@@ -119,7 +124,6 @@ void UCameraNodalOffsetAlgoOpticalAxis::Tick(float DeltaTime)
 		FCrosshairOverlayParams OverlayParams;
 		OverlayParams.PrincipalPoint = FVector2f(DistortionState.ImageCenter.PrincipalPoint);
 		OverlayRendering::DrawCrosshairOverlay(OverlayTexture, OverlayParams);
-		StepsController->RefreshOverlay();
 
 		LastPrincipalPoint = DistortionState.ImageCenter.PrincipalPoint;
 	}
diff --git a/Engine/Plugins/VirtualProduction/CameraCalibration/Source/CameraCalibrationEditor/Private/UI/CameraCalibrationWidgetHelpers.cpp b/Engine/Plugins/VirtualProduction/CameraCalibration/Source/CameraCalibrationEditor/Private/UI/CameraCalibrationWidgetHelpers.cpp
--- a/Engine/Plugins/VirtualProduction/CameraCalibration/Source/CameraCalibrationEditor/Private/UI/CameraCalibrationWidgetHelpers.cpp
+++ b/Engine/Plugins/VirtualProduction/CameraCalibration/Source/CameraCalibrationEditor/Private/UI/CameraCalibrationWidgetHelpers.cpp
@@ -77,7 +77,8 @@ void FCameraCalibrationWidgetHelpers::DisplayTextureInWindowAlmostFullScreen(UTe
 			.MinDesiredWidth(DetectionWindowMaxWidth)
 			.MinDesiredHeight(DetectionWindowMaxHeight)
 			[
-				SNew(SSimulcamViewport, Texture)
+				SNew(SSimulcamViewport)
+				.OverrideTexture(Texture)
 			]
 		]
 		.Buttons
