TrackingAlignmentTool: Adding TrackingAlignment Tool into the Camera Calibration plugin.

#virtualized

diff --git a/Engine/Plugins/VirtualProduction/CameraCalibration/CameraCalibration.uplugin b/Engine/Plugins/VirtualProduction/CameraCalibration/CameraCalibration.uplugin
--- a/Engine/Plugins/VirtualProduction/CameraCalibration/CameraCalibration.uplugin
+++ b/Engine/Plugins/VirtualProduction/CameraCalibration/CameraCalibration.uplugin
@@ -18,6 +18,16 @@
 			"Name": "CameraCalibrationEditor",
 			"Type": "Editor",
 			"LoadingPhase": "Default"
+		},
+		{
+			"Name": "TrackingAlignment",
+			"Type": "Runtime",
+			"LoadingPhase": "Default"
+		},
+		{
+			"Name": "TrackingAlignmentEditor",
+			"Type": "Editor",
+			"LoadingPhase": "Default"
 		}
 	],
 	"Plugins": [
diff --git a/Engine/Plugins/VirtualProduction/CameraCalibration/Source/TrackingAlignment/Private/TrackingAlignmentActors.h b/Engine/Plugins/VirtualProduction/CameraCalibration/Source/TrackingAlignment/Private/TrackingAlignmentActors.h
new file mode 100644
--- /dev/null
+++ b/Engine/Plugins/VirtualProduction/CameraCalibration/Source/TrackingAlignment/Private/TrackingAlignmentActors.h
@@ -0,0 +1,22 @@
+// Copyright Epic Games, Inc. All Rights Reserved.
+
+#pragma once
+
+#include "CoreMinimal.h"
+
+#include "TrackingAlignmentActors.generated.h"
+
+/** Struct of Source and Origin actors from a single Tracking space used for Space alignment. */
+USTRUCT(BlueprintType)
+struct FTrackingAlignmentActors
+{
+	GENERATED_BODY()
+
+	/** Tracked source actor in the tracking space. */
+	UPROPERTY(BlueprintReadWrite, EditAnywhere, Category = "Tracking Alignment")
+	TSoftObjectPtr<AActor> SourceActor;
+
+	/** Tracked origin actor in the tracking space. */
+	UPROPERTY(BlueprintReadWrite, EditAnywhere, Category = "Tracking Alignment")
+	TSoftObjectPtr<AActor> OriginActor;
+};
diff --git a/Engine/Plugins/VirtualProduction/CameraCalibration/Source/TrackingAlignment/Private/TrackingAlignmentBPLibrary.cpp b/Engine/Plugins/VirtualProduction/CameraCalibration/Source/TrackingAlignment/Private/TrackingAlignmentBPLibrary.cpp
new file mode 100644
--- /dev/null
+++ b/Engine/Plugins/VirtualProduction/CameraCalibration/Source/TrackingAlignment/Private/TrackingAlignmentBPLibrary.cpp
@@ -0,0 +1,229 @@
+// Copyright Epic Games, Inc. All Rights Reserved.
+
+#include "TrackingAlignmentBPLibrary.h"
+#include "GameFramework/Actor.h"
+#include "TrackingAlignmentActors.h"
+#include "TrackingAlignmentCalibrationProfile.h"
+#include "OpenCVHelper.h"
+
+#if WITH_OPENCV
+// IWYU pragma: begin_keep
+#include "PreOpenCVHeaders.h"
+#include <opencv2/calib3d.hpp>
+#include "PostOpenCVHeaders.h"
+// IWYU pragma: end_keep
+#endif
+
+#define LOCTEXT_NAMESPACE "TrackingAlignmentBPLibrary"
+
+namespace UE::TrackingAlignmentFunctionLibrary
+{
+
+#if WITH_OPENCV
+
+static int32 HandEyeCalibrationSolveMethod = static_cast<int32>(cv::HandEyeCalibrationMethod::CALIB_HAND_EYE_PARK);
+
+void ExtractOpenCVTranslationAndRotation(const FTransform& InTransform, cv::Mat& OutRot, cv::Mat& OutTra)
+{
+	// Convert the input FTransform to OpenCV's coordinate system
+	FTransform CvTransform = InTransform;
+	FOpenCVHelper::ConvertUnrealToOpenCV(CvTransform);
+
+	const FMatrix TransformationMatrix = CvTransform.ToMatrixNoScale();
+
+	// Extract the translation vector from the transformation matrix
+	OutTra = cv::Mat(3, 1, CV_64FC1);
+	OutTra.at<double>(0) = TransformationMatrix.M[3][0];
+	OutTra.at<double>(1) = TransformationMatrix.M[3][1];
+	OutTra.at<double>(2) = TransformationMatrix.M[3][2];
+
+	// Extract the rotation matrix from the transformation matrix
+	OutRot = cv::Mat(3, 3, CV_64FC1);
+	for (int32 Column = 0; Column < 3; ++Column)
+	{
+		const FVector ColumnVector = TransformationMatrix.GetColumn(Column);
+
+		// Note that since Unreal matrices stores basis vectors in rows, we need to transpose them into the Cv column-vector rotation matrix.
+
+		OutRot.at<double>(Column, 0) = ColumnVector.X;
+		OutRot.at<double>(Column, 1) = ColumnVector.Y;
+		OutRot.at<double>(Column, 2) = ColumnVector.Z;
+	}
+}
+
+void ExtractUnrealTransform(const cv::Mat& InRotation, const cv::Mat& InTranslation, FTransform& OutTransform)
+{
+	FMatrix TransformationMatrix = FMatrix::Identity;
+
+	// Add the rotation matrix to the transformation matrix
+	// Note: We need to transpose when converting from OpenCV back to Unreal format due to their different basis vector convention.
+	for (int32 Column = 0; Column < 3; ++Column)
+	{
+		TransformationMatrix.SetColumn(Column, FVector(InRotation.at<double>(Column, 0), InRotation.at<double>(Column, 1), InRotation.at<double>(Column, 2)));
+	}
+
+	// Add the translation vector to the transformation matrix
+	TransformationMatrix.M[3][0] = InTranslation.at<double>(0);
+	TransformationMatrix.M[3][1] = InTranslation.at<double>(1);
+	TransformationMatrix.M[3][2] = InTranslation.at<double>(2);
+
+	OutTransform.SetFromMatrix(TransformationMatrix);
+
+	// Convert the output FTransform to UE's coordinate system
+	FOpenCVHelper::ConvertOpenCVToUnreal(OutTransform);
+}
+
+bool GetSamplesAsEyeInHandMats(
+	UTrackingAlignmentCalibrationProfile* InCalibrationProfile,
+	std::vector<cv::Mat>& RGripper2Base,
+	std::vector<cv::Mat>& TGripper2Base,
+	std::vector<cv::Mat>& RTarget2Cam,
+	std::vector<cv::Mat>& TTarget2Cam
+)
+{
+	if (InCalibrationProfile == nullptr || InCalibrationProfile->Samples.Num() < UTrackingAlignmentFunctionLibrary::GetMinimumRequiredTrackerAligmentSampleCount())
+	{
+		return false;
+	}
+
+	const FTrackingAlignmentActors& TrackerAActors = InCalibrationProfile->TrackerAActors;
+	const FTrackingAlignmentActors& TrackerBActors = InCalibrationProfile->TrackerBActors;
+
+	for (const FTrackingAlignmentSample& ThisSample : InCalibrationProfile->Samples)
+	{
+		cv::Mat TrackerAOpenCVRotation    = cv::Mat(3, 3, CV_64FC1);
+		cv::Mat TrackerAOpenCVTranslation = cv::Mat(3, 1, CV_64FC1);
+
+		cv::Mat TrackerBOpenCVRotation    = cv::Mat(3, 3, CV_64FC1);
+		cv::Mat TrackerBOpenCVTranslation = cv::Mat(3, 1, CV_64FC1);
+
+		ExtractOpenCVTranslationAndRotation(ThisSample.TransformA, TrackerAOpenCVRotation, TrackerAOpenCVTranslation);
+		ExtractOpenCVTranslationAndRotation(ThisSample.TransformB, TrackerBOpenCVRotation, TrackerBOpenCVTranslation);
+
+		RGripper2Base.push_back(TrackerBOpenCVRotation);
+		TGripper2Base.push_back(TrackerBOpenCVTranslation);
+		RTarget2Cam.push_back(TrackerAOpenCVRotation);
+		TTarget2Cam.push_back(TrackerAOpenCVTranslation);
+	}
+
+	return true;
+}
+
+FTransform GetCam2Base(UTrackingAlignmentCalibrationProfile* InCalibrationProfile)
+{
+	/*
+	* U = Unreal/absolute coordinate system (3d scene origin).
+	* b = OCELLUS coordinate system relative to U.
+	* a = OPTITRACK coordinate system relative to U.
+	* o = OCELLUS i-th position relative to Lo.
+	* m = MARKERTREE i-th position relative to To.
+	* Cam2Base = Rigid transform of MARKERTREE relative to OCELLUS.
+	*/
+	FTransform Cam2Base = FTransform::Identity;
+
+	std::vector<cv::Mat> RGripper2Base;
+	std::vector<cv::Mat> TGripper2Base;
+	std::vector<cv::Mat> RTarget2Cam;
+	std::vector<cv::Mat> TTarget2Cam;
+
+	if (!GetSamplesAsEyeInHandMats(InCalibrationProfile, RGripper2Base, TGripper2Base, RTarget2Cam, TTarget2Cam))
+	{
+		return FTransform::Identity;
+	}
+
+	// Pass off to OpenCV to solve for X
+	cv::Mat OutRot = cv::Mat::eye(3, 3, CV_64FC1);
+	cv::Mat OutTra = cv::Mat::zeros(3, 1, CV_64FC1);
+
+	const cv::HandEyeCalibrationMethod CalibrationMethod = static_cast<cv::HandEyeCalibrationMethod>(HandEyeCalibrationSolveMethod);
+
+	// For eye-to-hand configuration, invert gripper transforms as base2gripper
+	std::vector<cv::Mat> RBase2Gripper;
+	std::vector<cv::Mat> TBase2Gripper;
+	for (int32 SampleIdx = 0; SampleIdx < RGripper2Base.size(); ++SampleIdx)
+	{
+		cv::Mat InvertedR;
+		cv::Mat InvertedT;
+		cv::transpose(RGripper2Base[SampleIdx], InvertedR);
+		InvertedT = -InvertedR * TGripper2Base[SampleIdx];   // Tinv = -(Rinv * T)
+
+		RBase2Gripper.push_back(InvertedR);
+		TBase2Gripper.push_back(InvertedT);
+	}
+
+	cv::calibrateHandEye(RBase2Gripper, TBase2Gripper, RTarget2Cam, TTarget2Cam, OutRot, OutTra, CalibrationMethod);
+
+	ExtractUnrealTransform(OutRot, OutTra, Cam2Base);
+
+	return Cam2Base;
+}
+
+#else // WITH_OPENCV
+
+static int32 HandEyeCalibrationSolveMethod = 0;
+
+#endif // WITH_OPENCV
+
+} // namespace UE::TrackingAlignmentFunctionLibrary
+
+static FAutoConsoleVariableRef CVarHandEyeCalibrationSolveMethod(
+	TEXT("TrackingAlignment.SolveMethod"),
+	UE::TrackingAlignmentFunctionLibrary::HandEyeCalibrationSolveMethod,
+	TEXT("Set the OpenCV HandEye solve method to use when aligning tracking spaces.\n")
+	TEXT("0 - TSAI. A New Technique for Fully Autonomous and Efficient 3D Robotics Hand / Eye Calibration @cite Tsai89.\n")
+	TEXT("1 - PARK. Robot Sensor Calibration: Solving AX = XB on the Euclidean Group @cite Park94.\n")
+	TEXT("2 - HORAUD. Hand-eye Calibration @cite Horaud95.\n")
+	TEXT("3 - ANDREFF. On-line Hand-Eye Calibration @cite Andreff99.\n")
+	TEXT("4 - DANIILIDIS. Hand-Eye Calibration Using Dual Quaternions @cite Daniilidis98.")
+);
+
+int32 UTrackingAlignmentFunctionLibrary::GetMinimumRequiredTrackerAligmentSampleCount()
+{
+	return 5;
+}
+
+FTransform UTrackingAlignmentFunctionLibrary::GetAlignedTrackerBOrigin(UTrackingAlignmentCalibrationProfile* InCalibrationProfile)
+{
+	if (InCalibrationProfile == nullptr)
+	{
+		return FTransform::Identity;
+	}
+	
+	FTransform Cam2Base = FTransform::Identity;
+
+#if WITH_OPENCV
+	Cam2Base = UE::TrackingAlignmentFunctionLibrary::GetCam2Base(InCalibrationProfile);
+	if (InCalibrationProfile->TrackerAActors.OriginActor.IsNull())
+	{
+		return Cam2Base.Inverse();
+	}
+#endif // WITH_OPENCV
+
+	return Cam2Base.Inverse() * InCalibrationProfile->TrackerAActors.OriginActor->GetTransform();
+}
+
+bool UTrackingAlignmentFunctionLibrary::FindAndUpdateOriginActor(FTrackingAlignmentActors& InTrackingActors, AActor*& OutNewParentActor)
+{
+	if (InTrackingActors.SourceActor.IsNull())
+	{
+		return false;
+	}
+
+	AActor* SourceActorReal = InTrackingActors.SourceActor.LoadSynchronous();
+
+	if (SourceActorReal == nullptr)
+	{
+		return false;
+	}
+
+	if (AActor* AttachParentActor = SourceActorReal->GetAttachParentActor())
+	{
+		InTrackingActors.OriginActor = AttachParentActor;
+		OutNewParentActor = AttachParentActor;
+		return true;
+	}
+
+	return false;
+}
+
+#undef LOCTEXT_NAMESPACE
diff --git a/Engine/Plugins/VirtualProduction/CameraCalibration/Source/TrackingAlignment/Private/TrackingAlignmentBPLibrary.h b/Engine/Plugins/VirtualProduction/CameraCalibration/Source/TrackingAlignment/Private/TrackingAlignmentBPLibrary.h
new file mode 100644
--- /dev/null
+++ b/Engine/Plugins/VirtualProduction/CameraCalibration/Source/TrackingAlignment/Private/TrackingAlignmentBPLibrary.h
@@ -0,0 +1,36 @@
+// Copyright Epic Games, Inc. All Rights Reserved.
+
+#pragma once
+
+#include "CoreMinimal.h"
+#include "Kismet/BlueprintFunctionLibrary.h"
+#include "Math/Transform.h"
+
+#include "TrackingAlignmentBPLibrary.generated.h"
+
+struct FTrackingAlignmentActor;
+class UTrackingAlignmentCalibrationProfile;
+
+/** Blueprint functions to caluclate tracking space alignment. */
+UCLASS(MinimalAPI)
+class UTrackingAlignmentFunctionLibrary : public UBlueprintFunctionLibrary
+{
+	GENERATED_BODY()
+
+public:
+	/** Gets the minimum number of require tracker alignment samples to be able to calculate a Cam2Base value. */
+	UFUNCTION(BlueprintPure, Category = "Tracking Alignment")
+	static int32 GetMinimumRequiredTrackerAligmentSampleCount();
+
+	/** 
+	* Given a UTrackingAlignmentCalibrationProfile, returns the calculated alignment origin for TrackerB space to Tracker A space.
+	* Uses OpenCV to solve the space alignment using Eye-To-Hand calibration with the given calibration method.
+	*/
+	UFUNCTION(BlueprintCallable, Category = "Tracking Alignment")
+	static FTransform GetAlignedTrackerBOrigin(UTrackingAlignmentCalibrationProfile* InCalibrationProfile);
+
+	/** Given an actor being tracked, find and update its origin actor to the direct parent actor. */
+	UFUNCTION(BlueprintCallable, Category = "Tracking Alignment")
+	static bool FindAndUpdateOriginActor(UPARAM(ref) FTrackingAlignmentActors& InTrackingActors, AActor*& OutNewParentActor);
+};
+
diff --git a/Engine/Plugins/VirtualProduction/CameraCalibration/Source/TrackingAlignment/Private/TrackingAlignmentCalibrationProfile.cpp b/Engine/Plugins/VirtualProduction/CameraCalibration/Source/TrackingAlignment/Private/TrackingAlignmentCalibrationProfile.cpp
new file mode 100644
--- /dev/null
+++ b/Engine/Plugins/VirtualProduction/CameraCalibration/Source/TrackingAlignment/Private/TrackingAlignmentCalibrationProfile.cpp
@@ -0,0 +1,53 @@
+// Copyright Epic Games, Inc. All Rights Reserved.
+
+#include "TrackingAlignmentCalibrationProfile.h"
+#include "GameFramework/Actor.h"
+
+bool UTrackingAlignmentCalibrationProfile::CaptureSample(FTrackingAlignmentSample& NewSample)
+{
+	if (TrackerAActors.SourceActor.IsNull() || TrackerBActors.SourceActor.IsNull())
+	{
+		return false;
+	}
+
+	auto GetCapturedTransform = [](const FTrackingAlignmentActors& InActors) -> FTransform
+		{
+			AActor* SourceActor = InActors.SourceActor.LoadSynchronous();
+
+			FTransform OutTransform = SourceActor->GetTransform();
+
+			if (!InActors.OriginActor.IsNull())
+			{
+				const AActor* OriginActor = InActors.OriginActor.LoadSynchronous();
+				OutTransform = OutTransform.GetRelativeTransform(OriginActor->GetTransform());
+			}
+
+			return OutTransform;
+		};
+
+	Modify();
+	NewSample.TransformA = GetCapturedTransform(TrackerAActors);
+	NewSample.TransformB = GetCapturedTransform(TrackerBActors);
+
+	Samples.Add(NewSample);
+
+	return true;
+}
+
+bool UTrackingAlignmentCalibrationProfile::RemoveSample(int32 AtIndex)
+{
+	if (AtIndex >= 0 && AtIndex < Samples.Num())
+	{
+		Modify();
+		Samples.RemoveAt(AtIndex);
+		return true;
+	}
+
+	return false;
+}
+
+void UTrackingAlignmentCalibrationProfile::ClearSamples()
+{
+	Modify();
+	Samples.Empty();
+}
diff --git a/Engine/Plugins/VirtualProduction/CameraCalibration/Source/TrackingAlignment/Private/TrackingAlignmentCalibrationProfile.h b/Engine/Plugins/VirtualProduction/CameraCalibration/Source/TrackingAlignment/Private/TrackingAlignmentCalibrationProfile.h
new file mode 100644
--- /dev/null
+++ b/Engine/Plugins/VirtualProduction/CameraCalibration/Source/TrackingAlignment/Private/TrackingAlignmentCalibrationProfile.h
@@ -0,0 +1,42 @@
+// Copyright Epic Games, Inc. All Rights Reserved.
+
+#pragma once
+
+#include "CoreMinimal.h"
+#include "Engine/DataAsset.h"
+#include "TrackingAlignmentActors.h"
+#include "TrackingAlignmentSample.h"
+
+#include "TrackingAlignmentCalibrationProfile.generated.h"
+
+/** Calibration profile holding actors and samples for aligning two tracking spaces. */
+UCLASS(MinimalAPI, BlueprintType, AutoExpandCategories = ("Tracking Alignment"))
+class UTrackingAlignmentCalibrationProfile : public UPrimaryDataAsset
+{
+	GENERATED_BODY()
+
+public:
+	/** Tracking space A tracked source and origin actor. */
+	UPROPERTY(BlueprintReadWrite, EditAnywhere, Category = "Tracking Alignment")
+	FTrackingAlignmentActors TrackerAActors;
+
+	/** Tracking space B tracked source and origin actor. */
+	UPROPERTY(BlueprintReadWrite, EditAnywhere, Category = "Tracking Alignment")
+	FTrackingAlignmentActors TrackerBActors;
+
+	/** Captured samples from both Tracking space A and B. */
+	UPROPERTY(BlueprintReadOnly, EditAnywhere, Category = "Tracking Alignment")
+	TArray<FTrackingAlignmentSample> Samples;
+
+	/** Capture and add a new sample for TrackerAActors and TrackerBActors, and write it into the out parameter NewSample. */
+	UFUNCTION(BlueprintCallable, Category = "Tracking Alignment")
+	bool CaptureSample(FTrackingAlignmentSample& NewSample);
+
+	/** Remove sample at the given index. */
+	UFUNCTION(BlueprintCallable, Category = "Tracking Alignment")
+	bool RemoveSample(int32 AtIndex);
+
+	/** Clear all current samples. */
+	UFUNCTION(BlueprintCallable, Category = "Tracking Alignment")
+	void ClearSamples();
+};
diff --git a/Engine/Plugins/VirtualProduction/CameraCalibration/Source/TrackingAlignment/Private/TrackingAlignmentModule.cpp b/Engine/Plugins/VirtualProduction/CameraCalibration/Source/TrackingAlignment/Private/TrackingAlignmentModule.cpp
new file mode 100644
--- /dev/null
+++ b/Engine/Plugins/VirtualProduction/CameraCalibration/Source/TrackingAlignment/Private/TrackingAlignmentModule.cpp
@@ -0,0 +1,5 @@
+// Copyright Epic Games, Inc. All Rights Reserved.
+
+#include "Modules/ModuleManager.h"
+
+IMPLEMENT_MODULE(FDefaultModuleImpl, TrackingAlignment);
diff --git a/Engine/Plugins/VirtualProduction/CameraCalibration/Source/TrackingAlignment/Private/TrackingAlignmentSample.h b/Engine/Plugins/VirtualProduction/CameraCalibration/Source/TrackingAlignment/Private/TrackingAlignmentSample.h
new file mode 100644
--- /dev/null
+++ b/Engine/Plugins/VirtualProduction/CameraCalibration/Source/TrackingAlignment/Private/TrackingAlignmentSample.h
@@ -0,0 +1,23 @@
+// Copyright Epic Games, Inc. All Rights Reserved.
+
+#pragma once
+
+#include "CoreMinimal.h"
+#include "Math/Transform.h"
+
+#include "TrackingAlignmentSample.generated.h"
+
+/** Single sample as part of a Tracking Alignment calibration, holding transform samples from both Tracking space A and B. */
+USTRUCT(BlueprintType)
+struct FTrackingAlignmentSample
+{
+	GENERATED_BODY()
+
+	/** Transform sample from Tracking Space A. */
+	UPROPERTY(BlueprintReadOnly, VisibleAnywhere, Category = "Tracking Alignment")
+	FTransform TransformA;
+
+	/** Transform sample from Tracking Space B. */
+	UPROPERTY(BlueprintReadOnly, VisibleAnywhere, Category = "Tracking Alignment")
+	FTransform TransformB;
+};
diff --git a/Engine/Plugins/VirtualProduction/CameraCalibration/Source/TrackingAlignment/TrackingAlignment.Build.cs b/Engine/Plugins/VirtualProduction/CameraCalibration/Source/TrackingAlignment/TrackingAlignment.Build.cs
new file mode 100644
--- /dev/null
+++ b/Engine/Plugins/VirtualProduction/CameraCalibration/Source/TrackingAlignment/TrackingAlignment.Build.cs
@@ -0,0 +1,28 @@
+// Copyright Epic Games, Inc. All Rights Reserved.
+
+using UnrealBuildTool;
+
+public class TrackingAlignment : ModuleRules
+{
+	public TrackingAlignment(ReadOnlyTargetRules Target) : base(Target)
+	{
+		PCHUsage = PCHUsageMode.UseExplicitOrSharedPCHs;
+
+		PublicDependencyModuleNames.AddRange(
+			new string[]
+			{
+				"Core",
+				"CoreUObject",
+				"Engine",
+				"OpenCVHelper",
+				"OpenCV",
+			}
+		);
+
+		PrivateDependencyModuleNames.AddRange(
+			new string[]
+			{
+			}
+		);
+	}
+}
diff --git a/Engine/Plugins/VirtualProduction/CameraCalibration/Source/TrackingAlignmentEditor/Private/TrackingAlignmentEditorModule.cpp b/Engine/Plugins/VirtualProduction/CameraCalibration/Source/TrackingAlignmentEditor/Private/TrackingAlignmentEditorModule.cpp
new file mode 100644
--- /dev/null
+++ b/Engine/Plugins/VirtualProduction/CameraCalibration/Source/TrackingAlignmentEditor/Private/TrackingAlignmentEditorModule.cpp
@@ -0,0 +1,60 @@
+// Copyright Epic Games, Inc. All Rights Reserved.
+
+#include "Modules/ModuleManager.h"
+#include "EditorUtilityWidgetBlueprint.h"
+#include "ToolMenus.h"
+
+#define LOCTEXT_NAMESPACE "FTrackingAlignmentEditorModule"
+
+class FTrackingAlignmentEditorModule : public IModuleInterface
+{
+public:
+	inline static const FName TrackerAlignmentPanelTabName = FName("TrackerAlignment");
+
+/** IModuleInterface implementation */
+	virtual void StartupModule() override
+	{
+		UToolMenus::RegisterStartupCallback(FSimpleMulticastDelegate::FDelegate::CreateRaw(this, &FTrackingAlignmentEditorModule::RegisterMenus));
+
+		FGlobalTabmanager::Get()->RegisterNomadTabSpawner(TrackerAlignmentPanelTabName, FOnSpawnTab::CreateRaw(this, &FTrackingAlignmentEditorModule::OnSpawnTrackerAlignment))
+			.SetDisplayName(NSLOCTEXT("TrackerAlignment", "TrackerAlignmentTabTitle", "Tracker Alignment"))
+			.SetTooltipText(NSLOCTEXT("TrackerAlignment", "TrackerAlignmentTooltipText", "Open the Tracker Alignment tool"))
+			.SetMenuType(ETabSpawnerMenuType::Hidden);
+	}
+
+	virtual void ShutdownModule() override
+	{
+		//Clean up nomad tab spawner
+		FGlobalTabmanager::Get()->UnregisterNomadTabSpawner(TrackerAlignmentPanelTabName);
+	}
+
+private:
+	TSharedRef<class SDockTab> OnSpawnTrackerAlignment(const class FSpawnTabArgs& SpawnTabArgs)
+	{
+		UEditorUtilityWidgetBlueprint* TrackerAlignmentEUW = LoadObject<UEditorUtilityWidgetBlueprint>(NULL, TEXT("/CameraCalibration/Widgets/EUW_TrackerAlignment.EUW_TrackerAlignment"), NULL, LOAD_None, NULL);
+		return TrackerAlignmentEUW->SpawnEditorUITab(SpawnTabArgs);
+	}
+
+	void RegisterMenus()
+	{
+		FToolMenuOwnerScoped OwnerScoped(this);
+
+		{
+			UToolMenu* Menu = UToolMenus::Get()->ExtendMenu("LevelEditor.MainMenu.Window.VirtualProduction");
+			{
+				FToolMenuSection& Section = Menu->FindOrAddSection("VirtualProduction");
+				Section.AddMenuEntry(
+					FName("TrackerAlignment"),
+					LOCTEXT("TrackingAlignmentMenuSpawnTrackerAlignment", "Tracker Alignment"),
+					LOCTEXT("TrackingAlignmentMenuSpawnTrackerAlignmentTooltip", "Open the Tracker Alignment tool"),
+					FSlateIcon(FAppStyle::Get().GetStyleSetName(), "Icons.Transform"),
+					FUIAction(FExecuteAction::CreateLambda([]() { FGlobalTabmanager::Get()->TryInvokeTab(TrackerAlignmentPanelTabName); }))
+				);
+			}
+		}
+	}
+};
+
+#undef LOCTEXT_NAMESPACE
+
+IMPLEMENT_MODULE(FTrackingAlignmentEditorModule, TrackingAlignmentEditor);
diff --git a/Engine/Plugins/VirtualProduction/CameraCalibration/Source/TrackingAlignmentEditor/TrackingAlignmentEditor.Build.cs b/Engine/Plugins/VirtualProduction/CameraCalibration/Source/TrackingAlignmentEditor/TrackingAlignmentEditor.Build.cs
new file mode 100644
--- /dev/null
+++ b/Engine/Plugins/VirtualProduction/CameraCalibration/Source/TrackingAlignmentEditor/TrackingAlignmentEditor.Build.cs
@@ -0,0 +1,32 @@
+// Copyright Epic Games, Inc. All Rights Reserved.
+
+using UnrealBuildTool;
+
+public class TrackingAlignmentEditor : ModuleRules
+{
+	public TrackingAlignmentEditor(ReadOnlyTargetRules Target) : base(Target)
+	{
+		PCHUsage = PCHUsageMode.UseExplicitOrSharedPCHs;
+
+		PublicDependencyModuleNames.AddRange(
+			new string[]
+			{
+				"Core",
+				"CoreUObject",
+				"Engine"
+			}
+		);
+
+		PrivateDependencyModuleNames.AddRange(
+			new string[]
+			{
+				"Blutility",
+				"Slate",
+				"SlateCore",
+				"ToolMenus",
+				"TrackingAlignment",
+				"UMGEditor"
+			}
+		);
+	}
+}
