CameraCalibration: Fix incorrect cropping of render when in PPM distortion render mode.

Without this, zooming in the simulcam viewer of the lens calibration tool would not behave correctly.

#rnx

diff --git a/Engine/Plugins/VirtualProduction/CameraCalibration/Source/CameraCalibrationEditor/Private/AssetEditor/Viewport/SimulcamEditorViewportClient.cpp b/Engine/Plugins/VirtualProduction/CameraCalibration/Source/CameraCalibrationEditor/Private/AssetEditor/Viewport/SimulcamEditorViewportClient.cpp
--- a/Engine/Plugins/VirtualProduction/CameraCalibration/Source/CameraCalibrationEditor/Private/AssetEditor/Viewport/SimulcamEditorViewportClient.cpp
+++ b/Engine/Plugins/VirtualProduction/CameraCalibration/Source/CameraCalibrationEditor/Private/AssetEditor/Viewport/SimulcamEditorViewportClient.cpp
@@ -3,6 +3,7 @@
 #include "AssetEditor/Viewport/SimulcamEditorViewportClient.h"
 #include "CanvasItem.h"
 #include "CanvasTypes.h"
+#include "CineCameraComponent.h"
 #include "Editor/UnrealEdEngine.h"
 #include "Framework/Application/SlateApplication.h"
 #include "ImageUtils.h"
@@ -363,7 +364,7 @@ void FSimulcamEditorViewportClient::PreDraw(FViewport* InViewport, FCanvas* Canv
 	{
 		// If the CG viewport size is zero, assume that either the media texture size is zero/invalid or an aspect ratio is invalid. Either way, skip
 		// setting the view rect and bounds in the crop material so that a normal viewport is rendered
-		if (CGViewportSize.IsNearlyZero())
+		if (CGViewportSize.X <= KINDA_SMALL_NUMBER || CGViewportSize.Y <= KINDA_SMALL_NUMBER)
 		{
 			CropToViewRectMaterialInstance->ClearParameterValues();
 			return;
@@ -374,22 +375,34 @@ void FSimulcamEditorViewportClient::PreDraw(FViewport* InViewport, FCanvas* Canv
 		{
 			if (ACameraActor* Camera = ElementsProvider.Pin()->GetCameraActor())
 			{
-				Overscan = Camera->GetCameraComponent()->Overscan + 1.0f;
+				if (UCineCameraComponent* CineCamera = Cast<UCineCameraComponent>(Camera->GetCameraComponent()))
+				{
+					// Only apply overscan if the render pipeline will crop it automatically
+					if (CineCamera->bCropOverscan)
+					{
+						Overscan = CineCamera->Overscan + 1.0f;
+
+						// This would be unexpected but best to keep it safe.
+						if (Overscan <= KINDA_SMALL_NUMBER)
+						{
+							Overscan = 1.0f;
+						}
+					}
+					// When bCropOverscan is false, Overscan remains 1.0f
+				}
 			}
 		}
 
 		// Compute the view rect from the CG viewport size and the current zoom/pan
 		const FBox2D CGViewRect((ViewRect.Min - 0.5) * ViewportSize / CGViewportSize + 0.5, (ViewRect.Max - 0.5) * ViewportSize / CGViewportSize + 0.5);
 		
-		// In order to crop and scale a distorted camera render that includes overscan, we must convert the view rect to overscanned coordinates,
-		// and also add a bit of overscan since the render pipeline will crop the overscan as the very last thing after post process materials are processed.
-		// For posterity, the equation below is derived from combining the overscan scaling equation (rect = (rect - center) * overscan + center) with
+		// Apply overscan transformation (when Overscan = 1.0, this is a no-op)
+		// The equation below is derived from combining the overscan scaling equation (rect = (rect - center) * overscan + center) with
 		// the equation to transform to overscanned coordinate space (rect = (rect - 0.5) / overscan + 0.5)
 		const FBox2D OverscannedViewRect = CGViewRect.ShiftBy((CGViewRect.GetCenter() - 0.5) * (1.0 / Overscan - 1.0));
 		
-		// Similarly to above, if the CG viewport is being inset at a smaller size than the whole viewport, we need to manually crop out the overscan as it will be too inset
-		// for the crop post process pass to crop it out
-		const FBox2D OverscanBounds(-FVector2D(0.5, 0.5) / Overscan + 0.5, FVector2D(0.5, 0.5) / Overscan + 0.5);
+		// Compute the overscan bounds (when Overscan = 1.0, this yields (0,0) to (1,1))
+		const FBox2D OverscanBounds = FBox2D(-FVector2D(0.5, 0.5) / Overscan + 0.5, FVector2D(0.5, 0.5) / Overscan + 0.5);
 		
 		CropToViewRectMaterialInstance->SetVectorParameterValue("ViewRect", FVector4(OverscannedViewRect.Min.X, OverscannedViewRect.Min.Y, OverscannedViewRect.Max.X, OverscannedViewRect.Max.Y));
 		CropToViewRectMaterialInstance->SetVectorParameterValue("ViewBounds", FVector4(OverscanBounds.Min.X, OverscanBounds.Min.Y, OverscanBounds.Max.X, OverscanBounds.Max.Y));
