CameraCal: Added missing Nodal Offset Algo files.

Removed flag to not display the pattern.
Fixed a clamping edge case.
Better handles texture lock failures.
Fixed wrong lens distortion when the charuco board is not at the origin.

#rnx

diff --git a/Engine/Plugins/Runtime/OpenCV/Source/OpenCVHelper/Private/OpenCVHelper.cpp b/Engine/Plugins/Runtime/OpenCV/Source/OpenCVHelper/Private/OpenCVHelper.cpp
--- a/Engine/Plugins/Runtime/OpenCV/Source/OpenCVHelper/Private/OpenCVHelper.cpp
+++ b/Engine/Plugins/Runtime/OpenCV/Source/OpenCVHelper/Private/OpenCVHelper.cpp
@@ -23,6 +23,52 @@
 
 DEFINE_LOG_CATEGORY_STATIC(LogOpenCVHelper, Log, All);
 
+#if WITH_OPENCV
+namespace
+{
+	/**
+	 * Safely copies data from an OpenCV Mat to a destination buffer.
+	 * Handles both contiguous and non-contiguous Mats.
+	 * 
+	 * @param SourceMat The source OpenCV Mat to copy from
+	 * @param DestBuffer The destination buffer to copy to
+	 * @param DestBufferSize The size of the destination buffer in bytes
+	 * @return true if copy succeeded, false if buffer was too small
+	 */
+	bool SafeCopyMatToBuffer(const cv::Mat& SourceMat, void* DestBuffer, SIZE_T DestBufferSize)
+	{
+		const SIZE_T RequiredBytes = SourceMat.total() * SourceMat.elemSize();
+		
+		if (RequiredBytes > DestBufferSize)
+		{
+			UE_LOG(LogOpenCVHelper, Error, TEXT("SafeCopyMatToBuffer: Destination buffer too small. Required: %llu bytes, Available: %llu bytes"), 
+				RequiredBytes, DestBufferSize);
+			return false;
+		}
+		
+		if (SourceMat.isContinuous())
+		{
+			// Mat is contiguous, can copy entire buffer at once
+			FMemory::Memcpy(DestBuffer, SourceMat.data, RequiredBytes);
+		}
+		else
+		{
+			// Mat is not contiguous (e.g., ROI view), copy row by row
+			const SIZE_T RowSizeBytes = SIZE_T(SourceMat.cols) * SourceMat.elemSize();
+			uint8* DestPtr = static_cast<uint8*>(DestBuffer);
+			
+			for (int32 Row = 0; Row < SourceMat.rows; ++Row)
+			{
+				const uint8* SrcPtr = SourceMat.ptr<uint8>(Row);
+				FMemory::Memcpy(DestPtr, SrcPtr, RowSizeBytes);
+				DestPtr += RowSizeBytes;
+			}
+		}
+		
+		return true;
+	}
+}
+#endif // WITH_OPENCV
 
 namespace UE::Aruco::Private
 {
@@ -165,7 +211,17 @@ UTexture2D* FOpenCVHelper::TextureFromCvMat(cv::Mat& Mat, const FString* Package
 		const int32 NumSlices = 1;
 		const int32 NumMips = 1;
 
-		Texture->Source.Init(Mat.cols, Mat.rows, NumSlices, NumMips, SourceFormat, Mat.data);
+		if (Mat.isContinuous())
+		{
+			Texture->Source.Init(Mat.cols, Mat.rows, NumSlices, NumMips, SourceFormat, Mat.data);
+		}
+		else
+		{
+			// Mat is not contiguous, create a contiguous copy for the UE API
+			cv::Mat ContiguousMat;
+			Mat.copyTo(ContiguousMat);
+			Texture->Source.Init(Mat.cols, Mat.rows, NumSlices, NumMips, SourceFormat, ContiguousMat.data);
+		}
 
 		auto IsPowerOfTwo = [](int32 Value)
 		{
@@ -223,14 +279,31 @@ UTexture2D* FOpenCVHelper::TextureFromCvMat(cv::Mat& Mat, const FString* Package
 		}
 
 		// Copy the pixels from the OpenCV Mat to the Texture
+		{
+			FTexture2DMipMap& Mip0 = Texture->GetPlatformData()->Mips[0];
+			void* TextureData = Mip0.BulkData.Lock(LOCK_READ_WRITE);
+			
+			ON_SCOPE_EXIT
+			{
+				if (Mip0.BulkData.IsLocked())
+				{
+					Mip0.BulkData.Unlock();
+				}
+			};
+			
+			if (!TextureData)
+			{
+				Texture->MarkAsGarbage();
+				return nullptr;
+			}
 
-		FTexture2DMipMap& Mip0 = Texture->GetPlatformData()->Mips[0];
-		void* TextureData = Mip0.BulkData.Lock(LOCK_READ_WRITE);
-
-		const int32 PixelStride = Mat.channels();
-		FMemory::Memcpy(TextureData, Mat.data, Mat.cols * Mat.rows * SIZE_T(PixelStride));
-
-		Mip0.BulkData.Unlock();
+			const SIZE_T BufferSize = Mip0.BulkData.GetBulkDataSize();
+			if (!SafeCopyMatToBuffer(Mat, TextureData, BufferSize))
+			{
+				Texture->MarkAsGarbage();
+				return nullptr;
+			}
+		}
 
 		Texture->UpdateResource();
 	}
@@ -279,14 +352,29 @@ UTexture2D* FOpenCVHelper::TextureFromCvMat(cv::Mat& Mat, UTexture2D* InTexture)
 	}
 
 	// Copy the pixels from the OpenCV Mat to the Texture
+	{
+		FTexture2DMipMap& Mip0 = InTexture->GetPlatformData()->Mips[0];
+		void* TextureData = Mip0.BulkData.Lock(LOCK_READ_WRITE);
+		
+		ON_SCOPE_EXIT
+		{
+			if (Mip0.BulkData.IsLocked())
+			{
+				Mip0.BulkData.Unlock();
+			}
+		};
+		
+		if (!TextureData)
+		{
+			return nullptr;
+		}
 
-	FTexture2DMipMap& Mip0 = InTexture->GetPlatformData()->Mips[0];
-	void* TextureData = Mip0.BulkData.Lock(LOCK_READ_WRITE);
-
-	const int32 PixelStride = Mat.channels();
-	FMemory::Memcpy(TextureData, Mat.data, Mat.cols * Mat.rows * SIZE_T(PixelStride));
-
-	Mip0.BulkData.Unlock();
+		const SIZE_T BufferSize = Mip0.BulkData.GetBulkDataSize();
+		if (!SafeCopyMatToBuffer(Mat, TextureData, BufferSize))
+		{
+			return nullptr;
+		}
+	}
 
 	InTexture->UpdateResource();
 
@@ -465,24 +553,37 @@ bool FOpenCVHelper::DrawArucoMarkers(const TArray<FArucoMarker>& Markers, UTextu
 		Corners.push_back(MarkerCorners);
 	}
 
-	FTexture2DMipMap& Mip0 = DebugTexture->GetPlatformData()->Mips[0];
-	void* TextureData = Mip0.BulkData.Lock(LOCK_READ_WRITE);
-
-	// Create a matrix header pointing to the raw texture data, copy the debug marker outlines into it, and update the texture resource
-	cv::Mat TextureMat = cv::Mat(DebugTexture->GetSizeY(), DebugTexture->GetSizeX(), CV_8UC4, TextureData);
+	{
+		FTexture2DMipMap& Mip0 = DebugTexture->GetPlatformData()->Mips[0];
+		void* TextureData = Mip0.BulkData.Lock(LOCK_READ_WRITE);
+		
+		ON_SCOPE_EXIT
+		{
+			if (Mip0.BulkData.IsLocked())
+			{
+				Mip0.BulkData.Unlock();
+			}
+		};
+		
+		if (!TextureData)
+		{
+			return false;
+		}
 
-	// We need to convert from RGBA to RGB temporarily because cv::aruco::drawDetectedMarkers only supports images with 1 or 3 color channels (not 4)
-	// We use a second matrix header because cvtColor will make a copy of the original data each time it changes the color format and we do not want to lose the reference to the original TextureData pointer.
-	cv::Mat DebugMat;
-	TextureMat.copyTo(DebugMat);
+		// Create a matrix header pointing to the raw texture data, copy the debug marker outlines into it, and update the texture resource
+		cv::Mat TextureMat = cv::Mat(DebugTexture->GetSizeY(), DebugTexture->GetSizeX(), CV_8UC4, TextureData);
 
-	cv::cvtColor(DebugMat, DebugMat, cv::COLOR_RGBA2RGB);
-	cv::aruco::drawDetectedMarkers(DebugMat, Corners, MarkerIds);
-	cv::cvtColor(DebugMat, DebugMat, cv::COLOR_RGB2RGBA);
+		// We need to convert from RGBA to RGB temporarily because cv::aruco::drawDetectedMarkers only supports images with 1 or 3 color channels (not 4)
+		// We use a second matrix header because cvtColor will make a copy of the original data each time it changes the color format and we do not want to lose the reference to the original TextureData pointer.
+		cv::Mat DebugMat;
+		TextureMat.copyTo(DebugMat);
 
-	DebugMat.copyTo(TextureMat);
+		cv::cvtColor(DebugMat, DebugMat, cv::COLOR_RGBA2RGB);
+		cv::aruco::drawDetectedMarkers(DebugMat, Corners, MarkerIds);
+		cv::cvtColor(DebugMat, DebugMat, cv::COLOR_RGB2RGBA);
 
-	Mip0.BulkData.Unlock();
+		DebugMat.copyTo(TextureMat);
+	}
 
 	DebugTexture->UpdateResource();
 
@@ -517,8 +618,8 @@ bool FOpenCVHelper::IdentifyCheckerboard(TArray<FColor>& Image, FIntPoint ImageS
 	// Sanitize the ROI to ensure that it lies completely within the bounds of the full size image
 	RegionOfInterest.Min.X = FMath::Clamp(RegionOfInterest.Min.X, 0, ImageSize.X);
 	RegionOfInterest.Min.Y = FMath::Clamp(RegionOfInterest.Min.Y, 0, ImageSize.Y);
-	RegionOfInterest.Max.X = FMath::Clamp(RegionOfInterest.Max.X, 0, ImageSize.X);
-	RegionOfInterest.Max.Y = FMath::Clamp(RegionOfInterest.Max.Y, 0, ImageSize.Y);
+	RegionOfInterest.Max.X = FMath::Clamp(RegionOfInterest.Max.X, RegionOfInterest.Min.X + 1, ImageSize.X);
+	RegionOfInterest.Max.Y = FMath::Clamp(RegionOfInterest.Max.Y, RegionOfInterest.Min.Y + 1, ImageSize.Y);
 
 	if ((RegionOfInterest.Width() <= 0) || (RegionOfInterest.Height() <= 0))
 	{
@@ -587,17 +688,30 @@ bool FOpenCVHelper::DrawCheckerboardCorners(const TArray<FVector2f>& Corners, FI
 
 	cv::Mat CornerMat(Corners.Num(), 1, CV_32FC2, (void*)Corners.GetData());
 
-	FTexture2DMipMap& Mip0 = DebugTexture->GetPlatformData()->Mips[0];
-	void* TextureData = Mip0.BulkData.Lock(LOCK_READ_WRITE);
-
-	// Create a matrix header pointing to the raw texture data, draw the debug pattern, and update the texture resource
-	const FIntPoint TextureSize = FIntPoint(DebugTexture->GetSizeX(), DebugTexture->GetSizeY());
-	cv::Mat TextureMat = cv::Mat(TextureSize.Y, TextureSize.X, CV_8UC4, TextureData);
+	{
+		FTexture2DMipMap& Mip0 = DebugTexture->GetPlatformData()->Mips[0];
+		void* TextureData = Mip0.BulkData.Lock(LOCK_READ_WRITE);
+		
+		ON_SCOPE_EXIT
+		{
+			if (Mip0.BulkData.IsLocked())
+			{
+				Mip0.BulkData.Unlock();
+			}
+		};
+		
+		if (!TextureData)
+		{
+			return false;
+		}
 
-	constexpr bool bFoundPattern = true;
-	cv::drawChessboardCorners(TextureMat, cv::Size(CheckerboardDimensions.X, CheckerboardDimensions.Y), CornerMat, bFoundPattern);
+		// Create a matrix header pointing to the raw texture data, draw the debug pattern, and update the texture resource
+		const FIntPoint TextureSize = FIntPoint(DebugTexture->GetSizeX(), DebugTexture->GetSizeY());
+		cv::Mat TextureMat = cv::Mat(TextureSize.Y, TextureSize.X, CV_8UC4, TextureData);
 
-	Mip0.BulkData.Unlock();
+		constexpr bool bFoundPattern = true;
+		cv::drawChessboardCorners(TextureMat, cv::Size(CheckerboardDimensions.X, CheckerboardDimensions.Y), CornerMat, bFoundPattern);
+	}
 
 	DebugTexture->UpdateResource();
 
@@ -1004,8 +1118,8 @@ bool FOpenCVHelper::IdentifyCharucoCorners(TArray<FColor>& Image, FIntPoint Imag
 
 		int ROILeft = FMath::Clamp(ROI.Min.X, 0, ImageSize.X - 1);
 		int ROITop = FMath::Clamp(ROI.Min.Y, 0, ImageSize.Y - 1);
-		int ROIRight = FMath::Clamp(ROI.Max.X, ROI.Min.X + 1, ImageSize.X);
-		int ROIBottom = FMath::Clamp(ROI.Max.Y, ROI.Min.Y + 1, ImageSize.Y);
+		int ROIRight = FMath::Clamp(ROI.Max.X, ROILeft + 1, ImageSize.X);
+		int ROIBottom = FMath::Clamp(ROI.Max.Y, ROITop + 1, ImageSize.Y);
 		
 		ROIOffset = FIntPoint(ROILeft, ROITop);
 		cv::Rect ROIRect(ROILeft, ROITop, ROIRight - ROILeft, ROIBottom - ROITop);
@@ -1129,18 +1243,20 @@ bool FOpenCVHelper::DrawCharucoCorners(const FCharucoCorners& Corners, UTexture2
 	FTexture2DMipMap& Mip = PlatData->Mips[0];
 	uint8* TextureData = static_cast<uint8*>(Mip.BulkData.Lock(LOCK_READ_WRITE));
 
+	ON_SCOPE_EXIT
+	{
+		if (Mip.BulkData.IsLocked())
+		{
+			Mip.BulkData.Unlock();
+		}
+	};
+
 	if (!TextureData)
 	{
 		UE_LOG(LogOpenCVHelper, Error, TEXT("DrawCharucoCorners: Failed to lock texture bulk data"));
 		return false;
 	}
 
-	// Ensure all exit paths from now on call Mip.BulkData.Unlock().
-	ON_SCOPE_EXIT
-	{
-		Mip.BulkData.Unlock();
-	};
-
 	const int32 BytesPerPixel = 4;
 
 	const SIZE_T ExpectedSize = SIZE_T(Width) * SIZE_T(Height) * SIZE_T(BytesPerPixel);
@@ -1150,6 +1266,8 @@ bool FOpenCVHelper::DrawCharucoCorners(const FCharucoCorners& Corners, UTexture2
 	{
 		UE_LOG(LogOpenCVHelper, Error, TEXT("DrawCharucoCorners: BulkData too small. Expected >= %llu, got %llu"),
 			(unsigned long long)ExpectedSize, (unsigned long long)BulkSize);
+
+		Mip.BulkData.Unlock();
 		return false;
 	}
 
@@ -1166,7 +1284,9 @@ bool FOpenCVHelper::DrawCharucoCorners(const FCharucoCorners& Corners, UTexture2
 
 	if (!bHasCharuco && !bHasMarkers)
 	{
+		Mip.BulkData.Unlock();
 		DebugTexture->UpdateResource();
+
 		return true;
 	}
 
@@ -1174,6 +1294,8 @@ bool FOpenCVHelper::DrawCharucoCorners(const FCharucoCorners& Corners, UTexture2
 	{
 		UE_LOG(LogOpenCVHelper, Error, TEXT("DrawCharucoCorners: corners/ids size mismatch (%d vs %d)"),
 			Corners.Corners.Num(), Corners.CornerIds.Num());
+
+		Mip.BulkData.Unlock();
 		return false;
 	}
 
@@ -1245,11 +1367,16 @@ bool FOpenCVHelper::DrawCharucoCorners(const FCharucoCorners& Corners, UTexture2
 	// Convert back to BGRA for Unreal texture (preserve alpha channel by expanding BGR)
 	cv::cvtColor(CvImageBGR, CvImageBGRA, cv::COLOR_BGR2BGRA);
 
+	Mip.BulkData.Unlock();
 	DebugTexture->UpdateResource();
+
 	return true;
+
 #else
+
 	UE_LOG(LogOpenCVHelper, Error, TEXT("DrawCharucoCorners: OpenCV not available (WITH_OPENCV=0)"));
 	return false;
+
 #endif
 }
 
@@ -1338,7 +1465,29 @@ UTexture2D* FOpenCVHelper::GenerateCharucoBoard(const FCharucoBoardConfig& Board
 	FTexture2DMipMap& Mip = NewTexture->GetPlatformData()->Mips[0];
 
 	uint8* TextureData = static_cast<uint8*>(Mip.BulkData.Lock(LOCK_READ_WRITE));
-	FMemory::Memcpy(TextureData, BoardImageBGRA.data, BoardImageBGRA.total() * BoardImageBGRA.elemSize());
+
+	ON_SCOPE_EXIT
+	{
+		if (Mip.BulkData.IsLocked())
+		{
+			Mip.BulkData.Unlock();
+		}
+	};
+
+	if (!TextureData)
+	{
+		UE_LOG(LogOpenCVHelper, Error, TEXT("GenerateCharucoBoard: Failed to lock texture bulk data"));
+
+		NewTexture->MarkAsGarbage();
+		return nullptr;
+	}
+	
+	const SIZE_T BufferSize = Mip.BulkData.GetBulkDataSize();
+	if (!SafeCopyMatToBuffer(BoardImageBGRA, TextureData, BufferSize))
+	{
+		NewTexture->MarkAsGarbage();
+		return nullptr;
+	}
 
 	Mip.BulkData.Unlock();
 	
diff --git a/Engine/Plugins/VirtualProduction/CameraCalibration/Source/CameraCalibrationEditor/CameraCalibrationEditor.Build.cs b/Engine/Plugins/VirtualProduction/CameraCalibration/Source/CameraCalibrationEditor/CameraCalibrationEditor.Build.cs
--- a/Engine/Plugins/VirtualProduction/CameraCalibration/Source/CameraCalibrationEditor/CameraCalibrationEditor.Build.cs
+++ b/Engine/Plugins/VirtualProduction/CameraCalibration/Source/CameraCalibrationEditor/CameraCalibrationEditor.Build.cs
@@ -6,6 +6,9 @@ namespace UnrealBuildTool.Rules
 	{
 		public CameraCalibrationEditor(ReadOnlyTargetRules Target) : base(Target)
 		{
+			// Enable C++ exceptions for try-catch blocks in calibration solver because it relies on OpenCV and it can produce them when calibrations fail.
+			bEnableExceptions = true;
+			
 			PrivateDependencyModuleNames.AddRange(
 				new string[]
 				{
diff --git a/Engine/Plugins/VirtualProduction/CameraCalibration/Source/CameraCalibrationEditor/Private/AssetEditor/LensDistortionTool.cpp b/Engine/Plugins/VirtualProduction/CameraCalibration/Source/CameraCalibrationEditor/Private/AssetEditor/LensDistortionTool.cpp
--- a/Engine/Plugins/VirtualProduction/CameraCalibration/Source/CameraCalibrationEditor/Private/AssetEditor/LensDistortionTool.cpp
+++ b/Engine/Plugins/VirtualProduction/CameraCalibration/Source/CameraCalibrationEditor/Private/AssetEditor/LensDistortionTool.cpp
@@ -32,6 +32,18 @@
 
 #define LOCTEXT_NAMESPACE "LensDistortionTool"
 
+
+// Console variable to control whether to generate object points in local space for Charuco calibration
+static TAutoConsoleVariable<bool> CVarCharucoUseLocalSpaceObjectPoints(
+	TEXT("CameraCalibration.Charuco.UseLocalSpaceObjectPoints"),
+	false,
+	TEXT("If true, generate object points in local space for Charuco calibration instead of using world space points.\n")
+	TEXT("The custom solver should not need object points in a canonical plane anymore but leaving this option here while it gets battle-testd\n")
+	TEXT("Default: false"),
+	ECVF_Default
+);
+
+
 namespace UE::CameraCalibration::Private::LensDistortionTool
 {
 	static const FString SessionDateTimeField(TEXT("SessionDateTime"));
@@ -776,6 +788,14 @@ void ULensDistortionTool::CalibrateLens()
 	{
 		GenerateDummyCheckerboardPoints(Samples3d, Dataset.CalibrationRows.Num(), Dataset.CalibrationRows[0]->CheckerboardDimensions);
 	}
+	else if (CaptureSettings.CalibrationPattern == ECalibrationPattern::Charuco)
+	{
+		// Custom solver should not need object points in a canonical plane anymore, but leaving this here while this is battle-tested.
+		if (CVarCharucoUseLocalSpaceObjectPoints.GetValueOnGameThread())
+		{
+			GenerateObjectPointsInLocalSpace(Samples3d);
+		}
+	}
 
 	if (CaptureSettings.bIsCalibratorTracked && CaptureSettings.CalibrationPattern != ECalibrationPattern::Points)
 	{
@@ -813,27 +833,42 @@ void ULensDistortionTool::CalibrateLens()
 
 	CalibrationTask = UE::Tasks::Launch(UE_SOURCE_LOCATION, [InSolver = this->Solver, Model, Samples3d, Samples2d, ImageSize, FocalLength, ImageCenter, DistortionGuess, CameraPoses, TargetPoses, PixelAspect, SolverFlags, LensFileEvalInputs]() mutable
 		{
-			FDistortionCalibrationResult Result = InSolver->Solve(
-				Samples3d,
-				Samples2d,
-				ImageSize,
-				FocalLength,
-				ImageCenter,
-				DistortionGuess.Parameters,
-				CameraPoses,
-				TargetPoses,
-				Model,
-				PixelAspect,
-				SolverFlags
-			);
-
-			// CalibrateCamera() returns focal length and image center in pixels, but the result is expected to be normalized by the image size
-			Result.FocalLength.FxFy = Result.FocalLength.FxFy / ImageSize;
-			Result.ImageCenter.PrincipalPoint = Result.ImageCenter.PrincipalPoint / ImageSize;
-
-			// FZ inputs to LUT
-			Result.EvaluatedFocus = LensFileEvalInputs.Focus;
-			Result.EvaluatedZoom = LensFileEvalInputs.Zoom;
+			FDistortionCalibrationResult Result;
+			
+			try
+			{
+				Result = InSolver->Solve(
+					Samples3d,
+					Samples2d,
+					ImageSize,
+					FocalLength,
+					ImageCenter,
+					DistortionGuess.Parameters,
+					CameraPoses,
+					TargetPoses,
+					Model,
+					PixelAspect,
+					SolverFlags
+				);
+
+				// CalibrateCamera() returns focal length and image center in pixels, but the result is expected to be normalized by the image size
+				Result.FocalLength.FxFy = Result.FocalLength.FxFy / ImageSize;
+				Result.ImageCenter.PrincipalPoint = Result.ImageCenter.PrincipalPoint / ImageSize;
+
+				// FZ inputs to LUT
+				Result.EvaluatedFocus = LensFileEvalInputs.Focus;
+				Result.EvaluatedZoom = LensFileEvalInputs.Zoom;
+			}
+			catch (const std::exception& Exception)
+			{
+				Result.ErrorMessage = FText::FromString(FString::Printf(TEXT("Calibration solver reported the error: %s"), ANSI_TO_TCHAR(Exception.what())));
+				UE_LOG(LogCameraCalibrationEditor, Error, TEXT("LensDistortionTool: Solver crashed with std::exception: %s"), ANSI_TO_TCHAR(Exception.what()));
+			}
+			catch (...)
+			{
+				Result.ErrorMessage = FText::FromString(TEXT("Calibration solver returned with an unspecified error"));
+				UE_LOG(LogCameraCalibrationEditor, Error, TEXT("LensDistortionTool: Solver crashed with unknown exception"));
+			}
 
 			return Result;
 		});
@@ -1064,9 +1099,10 @@ void ULensDistortionTool::RescalePoints(TArray<FVector2D>& Points, FIntPoint Deb
 
 void ULensDistortionTool::GenerateDummyCheckerboardPoints(TArray<FObjectPoints>& Samples3d, int32 NumImages, FIntPoint CheckerboardDimensions)
 {
-	// If the camera is not tracked, the distortion solver must initialize the camera pose for each image using linear algebra techniques.
-	// However, it struggles to do so when "real" tracking data is used for the calibrator. So in this case, we replace the tracked calibrator data 
+	// If the camera is not tracked, the distortion solver must initialize the camera pose for each image.
+	// However, it may struggle when "real" tracking data is used for the calibrator. So in this case, we replace the tracked calibrator data 
 	// with a set of dummy points for the 3D checkerboard corners. The board is assumed to lie in the YZ plane with the TopLeft corner at (0, 0, 0) in world space.
+
 	Samples3d.Empty();
 	for (int32 ImageIndex = 0; ImageIndex < NumImages; ++ImageIndex)
 	{
@@ -1083,6 +1119,31 @@ void ULensDistortionTool::GenerateDummyCheckerboardPoints(TArray<FObjectPoints>&
 	}
 }
 
+void ULensDistortionTool::GenerateObjectPointsInLocalSpace(TArray<FObjectPoints>& Samples3d)
+{
+	Samples3d.Empty();
+		
+	for (const TSharedPtr<FCalibrationRow>& Row : Dataset.CalibrationRows)
+	{
+		FObjectPoints Points3d;
+		
+		// Get the inverse of the calibrator transform that was used for this capture
+		const FTransform InverseTargetPose = Row->TargetPose.Inverse();
+		
+		for (int32 PointIdx = 0; PointIdx < Row->ObjectPoints.Points.Num(); ++PointIdx)
+		{
+			const FVector& WorldPoint = Row->ObjectPoints.Points[PointIdx];
+			
+			// Transform back to local space using the stored target pose
+			const FVector LocalPoint = InverseTargetPose.TransformPosition(WorldPoint);
+			
+			Points3d.Points.Add(LocalPoint);
+		}
+		
+		Samples3d.Add(Points3d);
+	}
+}
+
 bool ULensDistortionTool::DependsOnStep(UCameraCalibrationStep* Step) const
 {
 	return Cast<ULensInfoStep>(Step) != nullptr;
diff --git a/Engine/Plugins/VirtualProduction/CameraCalibration/Source/CameraCalibrationEditor/Private/AssetEditor/LensDistortionTool.h b/Engine/Plugins/VirtualProduction/CameraCalibration/Source/CameraCalibrationEditor/Private/AssetEditor/LensDistortionTool.h
--- a/Engine/Plugins/VirtualProduction/CameraCalibration/Source/CameraCalibrationEditor/Private/AssetEditor/LensDistortionTool.h
+++ b/Engine/Plugins/VirtualProduction/CameraCalibration/Source/CameraCalibrationEditor/Private/AssetEditor/LensDistortionTool.h
@@ -256,6 +256,9 @@ private:
 
 	/** Replace tracked checkerboard calibration points with a set of dummy 3D points */
 	void GenerateDummyCheckerboardPoints(TArray<FObjectPoints>& Samples3d, int32 NumImages, FIntPoint CheckerboardDimensions);
+	
+	/** Generates the object points in local space. Available in case the solver requires it. */
+	void GenerateObjectPointsInLocalSpace(TArray<FObjectPoints>& Samples3d);
 
 	/** Get the directory for the current session */
 	FString GetSessionSaveDir() const;
diff --git a/Engine/Plugins/VirtualProduction/CameraCalibration/Source/CameraCalibrationEditor/Private/Calibrators/CameraCalibrationSolver.cpp b/Engine/Plugins/VirtualProduction/CameraCalibration/Source/CameraCalibrationEditor/Private/Calibrators/CameraCalibrationSolver.cpp
--- a/Engine/Plugins/VirtualProduction/CameraCalibration/Source/CameraCalibrationEditor/Private/Calibrators/CameraCalibrationSolver.cpp
+++ b/Engine/Plugins/VirtualProduction/CameraCalibration/Source/CameraCalibrationEditor/Private/Calibrators/CameraCalibrationSolver.cpp
@@ -251,10 +251,13 @@ FDistortionCalibrationResult ULensDistortionSolverOpenCV::Solve_Implementation(
 		return Result;
 	}
 
+	// Compute canonical plane detection once to avoid redundant calculations
+	const bool bUseCanonicalPlaneInit = CanUseHomographyInitialization(ObjectPointsMat);
+	
 	// If the flag to use a starting guess for the camera intrinsics is not set, calculate some initial values for the intrinsic parameters
 	if (!(EnumHasAnyFlags(SolverFlags, ECalibrationFlags::UseIntrinsicGuess)))
 	{
-		InitCameraIntrinsics(ObjectPointsMat, ImagePointsMat, NumPointsMat, CvImageSize, CameraMatrix);
+		InitCameraIntrinsics(ObjectPointsMat, ImagePointsMat, NumPointsMat, CvImageSize, CameraMatrix, bUseCanonicalPlaneInit);
 	}
 
 	// Initialize the solver with the number of parameters to solve, and the maximum number of iterations to run
@@ -371,7 +374,7 @@ FDistortionCalibrationResult ULensDistortionSolverOpenCV::Solve_Implementation(
 
 			ObjectPointsIndex += NumImagePoints;
 
-			InitCameraExtrinsics(LensModel, ObjectPointsInImage, ImagePointsInImage, CameraMatrix, DistCoeffs, CvImageSize, Rotation, Translation, SolverFlags);
+			InitCameraExtrinsics(LensModel, ObjectPointsInImage, ImagePointsInImage, CameraMatrix, DistCoeffs, CvImageSize, Rotation, Translation, SolverFlags, bUseCanonicalPlaneInit);
 		}
 	}
 
@@ -670,12 +673,33 @@ double FCameraCalibrationSolver::OptimizeNodalOffset(
 
 #if WITH_OPENCV
 void ULensDistortionSolverOpenCV::InitCameraIntrinsics(
+	const cv::Mat& ObjectPoints,
+	const cv::Mat& ImagePoints,
+	const cv::Mat& NumPoints,
+	const cv::Size ImageSize,
+	cv::Mat& CameraMatrix,
+	bool bUseCanonicalPlaneInit)
+{	
+	if (bUseCanonicalPlaneInit)
+	{
+		UE_LOG(LogCameraCalibrationSolver, Verbose, TEXT("Using homography-based initialization for objects point in a canonical plane"));
+		InitIntrinsicsFromCanonicalPlane(ObjectPoints, ImagePoints, NumPoints, ImageSize, CameraMatrix);
+	}
+	else
+	{
+		UE_LOG(LogCameraCalibrationSolver, Verbose, TEXT("Using 3D point initialization for object points in general world coordinates"));
+		InitIntrinsicsFrom3DPoints(ObjectPoints, ImagePoints, NumPoints, ImageSize, CameraMatrix);
+	}
+}
+
+void ULensDistortionSolverOpenCV::InitIntrinsicsFromCanonicalPlane(
 	const cv::Mat& ObjectPoints,
 	const cv::Mat& ImagePoints,
 	const cv::Mat& NumPoints,
 	const cv::Size ImageSize,
 	cv::Mat& CameraMatrix)
 {
+	// Use homography-based method for 2D plane points (checkerboard, dummy points)
 	const int NumImages = NumPoints.cols;
 
 	const double Cx = (ImageSize.width - 1) * 0.5;
@@ -734,7 +758,7 @@ void ULensDistortionSolverOpenCV::InitCameraIntrinsics(
 		N[2] = FMath::InvSqrt(N[2]);
 		N[3] = FMath::InvSqrt(N[3]);
 
-		for (int Index = 0; Index < 3; ++Index)
+		for (int32 Index = 0; Index < 3; ++Index)
 		{
 			H[Index] *= N[0];
 			V[Index] *= N[1];
@@ -760,6 +784,18 @@ void ULensDistortionSolverOpenCV::InitCameraIntrinsics(
 
 	cv::solve(MatrixA, VecB, FxFyMat, cv::DECOMP_NORMAL + cv::DECOMP_SVD);
 
+	constexpr double MinFocalLength = 1e-6;
+
+	if (FxFy[0] < MinFocalLength || FxFy[1] < MinFocalLength)
+	{
+		UE_LOG(LogCameraCalibrationSolver, Warning, TEXT("InitIntrinsicsFromPlane2d: FxFy did not have reasonable values: [%f, %f]. "),
+			FxFy[0], FxFy[1]
+		);
+
+		FxFy[0] = FMath::Max(FxFy[0], MinFocalLength);
+		FxFy[1] = FMath::Max(FxFy[1], MinFocalLength);
+	}
+
 	const double Fx = FMath::Sqrt(FMath::Abs(1.0 / FxFy[0]));
 	const double Fy = FMath::Sqrt(FMath::Abs(1.0 / FxFy[1]));
 
@@ -778,7 +814,8 @@ void ULensDistortionSolverOpenCV::InitCameraExtrinsics(
 	const cv::Size ImageSize,
 	cv::Mat& Rotation,
 	cv::Mat& Translation,
-	const ECalibrationFlags SolverFlags)
+	const ECalibrationFlags SolverFlags,
+	bool bUseCanonicalPlaneInit)
 {
 	const int NumPoints = ObjectPoints.total();
 
@@ -808,8 +845,8 @@ void ULensDistortionSolverOpenCV::InitCameraExtrinsics(
 
 	cv::SVDecomp(MM, MatrixW, cv::noArray(), MatrixV, cv::SVD::MODIFY_A + cv::SVD::FULL_UV);
 
-	// a planar structure case (all object points lie in the same plane)
-	if (W[2] / W[1] < 1e-3)
+	// Use the same canonical plane detection as intrinsics initialization
+	if (bUseCanonicalPlaneInit)
 	{
 		cv::Mat RotationMatrix = MatrixV;
 
@@ -869,12 +906,45 @@ void ULensDistortionSolverOpenCV::InitCameraExtrinsics(
 
 		cv::Rodrigues(RotationMatrix, Rotation);
 	}
-	else // non-planar structure, not currently supported
+	else // Case of not being in a 2d canonical plane
 	{
-		cv::Mat RotationMatrix = cv::Mat::eye(3, 3, CV_64F);
-		cv::Rodrigues(RotationMatrix, Rotation);
-
-		Translation.setTo(cv::Scalar::all(0));
+		UE_LOG(LogCameraCalibrationSolver, Verbose, TEXT("Initializing extrinsics for non-planar 3D points using solvePnP"));
+		
+		// For 3D world coordinates, use cv::solvePnP to get initial pose estimate
+		// Convert ObjectPoints format: from (3 x NumPoints) to (NumPoints x 3)
+		cv::Mat ObjectPoints3x1 = ObjectPoints.reshape(1, NumPoints);
+		
+		// Convert ImagePoints format: from (1 x NumPoints, 2-channel) to (NumPoints x 1, 2-channel)  
+		cv::Mat ImagePoints2x1 = ImagePoints.reshape(2, NumPoints);
+		
+		try
+		{
+			bool bSolveSuccess = cv::solvePnP(
+				ObjectPoints3x1,    // (NumPoints x 3) object points
+				ImagePoints2x1,     // (NumPoints x 1, 2-channel) image points
+				CameraMatrix,       // 3x3 camera matrix
+				cv::Mat::zeros(5, 1, CV_64F), // No distortion for initial estimate
+				Rotation,           // Output: 3x1 rotation vector
+				Translation,        // Output: 3x1 translation vector
+				false,              // useExtrinsicGuess = false
+				cv::SOLVEPNP_ITERATIVE
+			);
+			
+			if (!bSolveSuccess)
+			{
+				UE_LOG(LogCameraCalibrationSolver, Warning, TEXT("solvePnP failed, using identity pose"));
+				cv::Mat RotationMatrix = cv::Mat::eye(3, 3, CV_64F);
+				cv::Rodrigues(RotationMatrix, Rotation);
+				Translation.setTo(cv::Scalar::all(0));
+			}
+		}
+		catch (const cv::Exception& e)
+		{
+			UE_LOG(LogCameraCalibrationSolver, Warning, TEXT("solvePnP exception: %s, using identity pose"), ANSI_TO_TCHAR(e.what()));
+			cv::Mat RotationMatrix = cv::Mat::eye(3, 3, CV_64F);
+			cv::Rodrigues(RotationMatrix, Rotation);
+			Translation.setTo(cv::Scalar::all(0));
+		}
 	}
 
 	// Refine extrinsic parameters
@@ -2055,6 +2125,109 @@ double FOptimizeNodalOffsetSolver::calc(const double* x) const
 
 	return RootMeanSquareError;
 }
+
+bool ULensDistortionSolverOpenCV::CanUseHomographyInitialization(const cv::Mat& ObjectPoints, double StdDevTolerance)
+{
+	if (ObjectPoints.channels() == 2)
+	{
+		return true;
+	}
+
+	if (ObjectPoints.channels() != 3)
+	{
+		return false;
+	}
+
+	// Check if any coordinate has low variance, indicating canonical plane (X=0, Y=0, or Z=0)
+	cv::Scalar Mean, StdDev;
+	cv::meanStdDev(ObjectPoints, Mean, StdDev);
+	
+	return (StdDev[0] < StdDevTolerance) || (StdDev[1] < StdDevTolerance) || (StdDev[2] < StdDevTolerance);
+}
+
+void ULensDistortionSolverOpenCV::InitIntrinsicsFrom3DPoints(
+	const cv::Mat& ObjectPoints,
+	const cv::Mat& ImagePoints,
+	const cv::Mat& NumPoints,
+	const cv::Size ImageSize,
+	cv::Mat& CameraMatrix)
+{
+	// For 3D world coordinates, we'll use cv::calibrateCamera to get initial intrinsic estimates
+	// This is more robust than homography-based methods for non XY plane points.
+	
+	const int32 NumImages = NumPoints.cols;
+	
+	// Convert the input matrices to the format expected by cv::calibrateCamera
+	std::vector<std::vector<cv::Point3f>> ObjectPointsVec;
+	std::vector<std::vector<cv::Point2f>> ImagePointsVec;
+	
+	ObjectPointsVec.reserve(NumImages);
+	ImagePointsVec.reserve(NumImages);
+	
+	int32 ObjectPointsIndex = 0;
+	for (int32 ImageIndex = 0; ImageIndex < NumImages; ++ImageIndex)
+	{
+		const int32 NumImagePoints = NumPoints.ptr<int32>()[ImageIndex];
+		
+		std::vector<cv::Point3f> ObjectPointsInImage;
+		std::vector<cv::Point2f> ImagePointsInImage;
+		
+		ObjectPointsInImage.reserve(NumImagePoints);
+		ImagePointsInImage.reserve(NumImagePoints);
+		
+		for (int32 PointIndex = 0; PointIndex < NumImagePoints; ++PointIndex)
+		{
+			const cv::Point3d& ObjectPoint = ObjectPoints.at<cv::Point3d>(0, ObjectPointsIndex + PointIndex);
+			const cv::Point2d& ImagePoint = ImagePoints.at<cv::Point2d>(0, ObjectPointsIndex + PointIndex);
+			
+			ObjectPointsInImage.push_back(cv::Point3f(ObjectPoint.x, ObjectPoint.y, ObjectPoint.z));
+			ImagePointsInImage.push_back(cv::Point2f(ImagePoint.x, ImagePoint.y));
+		}
+		
+		ObjectPointsVec.push_back(ObjectPointsInImage);
+		ImagePointsVec.push_back(ImagePointsInImage);
+		
+		ObjectPointsIndex += NumImagePoints;
+	}
+	
+	// Use cv::calibrateCamera to get initial estimates
+	cv::Mat DistCoeffs = cv::Mat::zeros(5, 1, CV_64F);
+	std::vector<cv::Mat> Rvecs, Tvecs;
+	
+	try
+	{
+		cv::calibrateCamera(
+			ObjectPointsVec,
+			ImagePointsVec,
+			ImageSize,
+			CameraMatrix,
+			DistCoeffs,
+			Rvecs,
+			Tvecs,
+			cv::CALIB_USE_INTRINSIC_GUESS
+		);
+		
+		UE_LOG(LogCameraCalibrationSolver, Log, TEXT("Initialized intrinsics from 3D points: Fx=%.2f, Fy=%.2f, Cx=%.2f, Cy=%.2f"), 
+			CameraMatrix.at<double>(0, 0), CameraMatrix.at<double>(1, 1), 
+			CameraMatrix.at<double>(0, 2), CameraMatrix.at<double>(1, 2));
+	}
+	catch (const cv::Exception& e)
+	{
+		UE_LOG(LogCameraCalibrationSolver, Warning, TEXT("cv::calibrateCamera failed for 3D point initialization: %s"), ANSI_TO_TCHAR(e.what()));
+		
+		// Fallback to basic initialization
+		const double Fx = ImageSize.width;
+		const double Fy = ImageSize.width; 
+		const double Cx = (ImageSize.width - 1) * 0.5;
+		const double Cy = (ImageSize.height - 1) * 0.5;
+		
+		CameraMatrix.at<double>(0, 0) = Fx;
+		CameraMatrix.at<double>(1, 1) = Fy;
+		CameraMatrix.at<double>(0, 2) = Cx;
+		CameraMatrix.at<double>(1, 2) = Cy;
+	}
+}
+
 #endif // WITH_OPENCV
 
 #undef LOCTEXT_NAMESPACE
diff --git a/Engine/Plugins/VirtualProduction/CameraCalibration/Source/CameraCalibrationEditor/Private/Calibrators/CameraCalibrationSolver.h b/Engine/Plugins/VirtualProduction/CameraCalibration/Source/CameraCalibrationEditor/Private/Calibrators/CameraCalibrationSolver.h
--- a/Engine/Plugins/VirtualProduction/CameraCalibration/Source/CameraCalibrationEditor/Private/Calibrators/CameraCalibrationSolver.h
+++ b/Engine/Plugins/VirtualProduction/CameraCalibration/Source/CameraCalibrationEditor/Private/Calibrators/CameraCalibrationSolver.h
@@ -175,7 +175,8 @@ private:
 		const cv::Mat& ImagePoints,
 		const cv::Mat& NumPoints,
 		const cv::Size ImageSize,
-		cv::Mat& CameraMatrix);
+		cv::Mat& CameraMatrix,
+		bool bUseCanonicalPlaneInit);
 
 	/** Initialize the camera extrinsics (rotation and translation vectors) for each image using linear algebra techniques */
 	void InitCameraExtrinsics(
@@ -187,7 +188,8 @@ private:
 		const cv::Size ImageSize,
 		cv::Mat& Rotation,
 		cv::Mat& Translation,
-		const ECalibrationFlags SolverFlags);
+		const ECalibrationFlags SolverFlags,
+		bool bUseCanonicalPlaneInit);
 
 	/** Project the input object points to 2D using the input camera intrinsics, extrinsics, and distortion parameters */
 	void ProjectPoints(
@@ -267,6 +269,25 @@ private:
 	/** Find the unique set of camera poses and the index into that set for each image */
 	void FindUniqueCameraPoses(const TArray<FTransform>& InCameraPoses, TArray<FTransform>& OutUniqueCameraPoses, TArray<int32>& OutUniquePoseIndices);
 
+	/** Check if the input points lie in a canonical plane suitable for homography */
+	bool CanUseHomographyInitialization(const cv::Mat& ObjectPoints, double Tolerance = 1e-6);
+
+	/** Initialize camera intrinsics from 2D plane points using homography */
+	void InitIntrinsicsFromCanonicalPlane(
+		const cv::Mat& ObjectPoints,
+		const cv::Mat& ImagePoints,
+		const cv::Mat& NumPoints,
+		const cv::Size ImageSize,
+		cv::Mat& CameraMatrix);
+
+	/** Initialize camera intrinsics from 3D world coordinates (not requiring canonical plane like homography) */
+	void InitIntrinsicsFrom3DPoints(
+		const cv::Mat& ObjectPoints,
+		const cv::Mat& ImagePoints,
+		const cv::Mat& NumPoints,
+		const cv::Size ImageSize,
+		cv::Mat& CameraMatrix);
+
 #endif	// WITH_OPENCV
 };
 
diff --git a/Engine/Plugins/VirtualProduction/CameraCalibration/Source/CameraCalibrationEditor/Private/Calibrators/CameraNodalOffsetAlgoCharuco.cpp b/Engine/Plugins/VirtualProduction/CameraCalibration/Source/CameraCalibrationEditor/Private/Calibrators/CameraNodalOffsetAlgoCharuco.cpp
new file mode 100644
--- /dev/null
+++ b/Engine/Plugins/VirtualProduction/CameraCalibration/Source/CameraCalibrationEditor/Private/Calibrators/CameraNodalOffsetAlgoCharuco.cpp
@@ -0,0 +1,424 @@
+// Copyright Epic Games, Inc. All Rights Reserved.
+
+#include "CameraNodalOffsetAlgoCharuco.h"
+#include "AssetEditor/CameraCalibrationStepsController.h"
+#include "AssetEditor/NodalOffsetTool.h"
+#include "CameraCalibrationCharucoBoard.h"
+#include "CameraCalibrationEditorLog.h"
+#include "CameraCalibrationUtilsPrivate.h"
+#include "Engine/TextureRenderTarget2D.h"
+#include "Engine/Texture2D.h"
+#include "EngineUtils.h"
+#include "Misc/MessageDialog.h"
+#include "OpenCVHelper.h"
+#include "TextureResource.h"
+#include "UI/CameraCalibrationWidgetHelpers.h"
+#include "UI/SFilterableActorPicker.h"
+#include "Widgets/Input/SCheckBox.h"
+#include "Widgets/Views/SListView.h"
+
+#define LOCTEXT_NAMESPACE "CameraNodalOffsetAlgoCharuco"
+
+bool UCameraNodalOffsetAlgoCharuco::OnViewportClicked(const FGeometry& MyGeometry, const FPointerEvent& MouseEvent)
+{
+	// We only respond to left clicks
+	if (MouseEvent.GetEffectingButton() != EKeys::LeftMouseButton)
+	{
+		return false;
+	}
+
+	FText ErrorMessage;
+	if (!PopulatePoints(ErrorMessage))
+	{
+		FMessageDialog::Open(EAppMsgType::Ok, ErrorMessage, LOCTEXT("CalibrationError", "Calibration Error"));
+	}
+
+	return true;
+}
+
+TSharedRef<SWidget> UCameraNodalOffsetAlgoCharuco::BuildUI()
+{
+	return SNew(SVerticalBox)
+
+		+ SVerticalBox::Slot() // Charuco board picker
+		.VAlign(EVerticalAlignment::VAlign_Top)
+		.AutoHeight()
+		.MaxHeight(FCameraCalibrationWidgetHelpers::DefaultRowHeight)
+		[FCameraCalibrationWidgetHelpers::BuildLabelWidgetPair(LOCTEXT("CharucoBoard", "Charuco Board"), BuildCharucoBoardPickerWidget())]
+
+		+ SVerticalBox::Slot() // Calibrator component picker
+		.VAlign(EVerticalAlignment::VAlign_Top)
+		.AutoHeight()
+		.MaxHeight(FCameraCalibrationWidgetHelpers::DefaultRowHeight)
+		[FCameraCalibrationWidgetHelpers::BuildLabelWidgetPair(LOCTEXT("CalibratorComponents", "Calibrator Component(s)"), BuildCalibrationComponentPickerWidget())]
+
+		+ SVerticalBox::Slot() // Show Detection
+		.VAlign(EVerticalAlignment::VAlign_Top)
+		.AutoHeight()
+		.MaxHeight(FCameraCalibrationWidgetHelpers::DefaultRowHeight)
+		[FCameraCalibrationWidgetHelpers::BuildLabelWidgetPair(LOCTEXT("ShowDetection", "Show Detection"), BuildShowDetectionWidget())]
+
+		+ SVerticalBox::Slot() // Calibration Rows
+		.AutoHeight()
+		.MaxHeight(12 * FCameraCalibrationWidgetHelpers::DefaultRowHeight)
+		[BuildCalibrationPointsTable()]
+
+		+ SVerticalBox::Slot() // Action buttons (e.g. Remove, Clear)
+		.HAlign(EHorizontalAlignment::HAlign_Center)
+		.AutoHeight()
+		.Padding(0, 20)
+		[BuildCalibrationActionButtons()]
+		;
+}
+
+TSharedRef<SWidget> UCameraNodalOffsetAlgoCharuco::BuildCharucoBoardPickerWidget()
+{
+	return SNew(SFilterableActorPicker)
+		.OnSetObject_Lambda([&](const FAssetData& AssetData) -> void
+		{
+			if (AssetData.IsValid())
+			{
+				SetCalibrator(Cast<ACameraCalibrationCharucoBoard>(AssetData.GetAsset()));
+			}
+		})
+		.OnShouldFilterAsset_Lambda([&](const FAssetData& AssetData) -> bool
+		{
+			return !!Cast<ACameraCalibrationCharucoBoard>(AssetData.GetAsset());
+		})
+		.ActorAssetData_Lambda([&]() -> FAssetData
+		{
+			return FAssetData(GetCalibrator(), true);
+		});
+}
+
+TSharedRef<SWidget> UCameraNodalOffsetAlgoCharuco::BuildShowDetectionWidget()
+{
+	return SNew(SCheckBox)
+		.IsChecked_Lambda([&]() -> ECheckBoxState
+		{
+			return bShouldShowDetectionWindow ? ECheckBoxState::Checked : ECheckBoxState::Unchecked;
+		})
+		.OnCheckStateChanged_Lambda([&](ECheckBoxState NewState) -> void
+		{
+			bShouldShowDetectionWindow = (NewState == ECheckBoxState::Checked);
+		});
+}
+
+bool UCameraNodalOffsetAlgoCharuco::OnViewportMarqueeSelect(FVector2D StartPosition, FVector2D EndPosition)
+{
+	// Calculate ROI rectangle from marquee selection
+	FIntRect MarqueeROI;
+	MarqueeROI.Min = FIntPoint(FMath::Floor(FMath::Min(StartPosition.X, EndPosition.X)), 
+							   FMath::Floor(FMath::Min(StartPosition.Y, EndPosition.Y)));
+	MarqueeROI.Max = FIntPoint(FMath::Floor(FMath::Max(StartPosition.X, EndPosition.X)), 
+							   FMath::Floor(FMath::Max(StartPosition.Y, EndPosition.Y)));
+
+
+	// Call PopulatePoints directly with the ROI - clean and efficient!
+	FText ErrorMessage;
+	bool bResult = PopulatePoints(ErrorMessage, MarqueeROI);
+	
+	if (!bResult)
+	{
+		UE_LOG(LogCameraCalibrationEditor, Warning, TEXT("CameraNodalOffsetAlgoCharuco: Marquee selection failed: %s"), *ErrorMessage.ToString());
+	}
+	
+	return bResult;
+}
+
+AActor* UCameraNodalOffsetAlgoCharuco::FindFirstCalibrator() const
+{
+	UWorld* World = GetWorld();
+	if (!World)
+	{
+		return nullptr;
+	}
+
+	for (TActorIterator<ACameraCalibrationCharucoBoard> ActorIterator(World); ActorIterator; ++ActorIterator)
+	{
+		return *ActorIterator;
+	}
+
+	return nullptr;
+}
+
+bool UCameraNodalOffsetAlgoCharuco::PopulatePoints(FText& OutErrorMessage)
+{
+	// Regular clicks use full image (empty ROI)
+	return PopulatePoints(OutErrorMessage, FIntRect());
+}
+
+bool UCameraNodalOffsetAlgoCharuco::PopulatePoints(FText& OutErrorMessage, const FIntRect& ROI)
+{
+	const FCameraCalibrationStepsController* StepsController;
+
+	if (!ensure(GetStepsControllerAndLensFile(&StepsController, nullptr)))
+	{
+		OutErrorMessage = LOCTEXT("ToolNotFound", "Tool not found");
+		return false;
+	}
+
+	if (!Calibrator.IsValid())
+	{
+		OutErrorMessage = LOCTEXT("InvalidCalibrator", "Please select a Charuco board calibrator in the given combo box.");
+		return false;
+	}
+
+	if (!LastCameraData.bIsValid)
+	{
+		OutErrorMessage = LOCTEXT("InvalidLastCameraData", "Could not find a cached set of camera data (e.g. FIZ). Check the Lens Component to make sure it has valid evaluation inputs.");
+		return false;
+	}
+
+	// Get the Charuco board
+	ACameraCalibrationCharucoBoard* CharucoBoard = Cast<ACameraCalibrationCharucoBoard>(Calibrator.Get());
+	if (!CharucoBoard)
+	{
+		OutErrorMessage = LOCTEXT("InvalidCharucoBoard", "The selected calibrator is not a Charuco board.");
+		return false;
+	}
+
+	// Read media pixels
+	TArray<FColor> Pixels;
+	FIntPoint Size;
+
+	if (!StepsController->ReadMediaPixels(Pixels, Size, OutErrorMessage, ESimulcamViewportPortion::CameraFeed))
+	{
+		return false;
+	}
+
+	// Setup board configuration from the actor (NumCorners + 1 = squares in each direction)
+	CachedBoardConfig.SquaresX = CharucoBoard->NumCornerCols + 1;
+	CachedBoardConfig.SquaresY = CharucoBoard->NumCornerRows + 1;
+	CachedBoardConfig.SquareSize = CharucoBoard->SquareSideLength;
+	CachedBoardConfig.MarkerSize = CharucoBoard->SquareSideLength * CharucoBoard->MarkerSizeRatio;
+	CachedBoardConfig.Dictionary = UE::CameraCalibration::Private::ShadowToOpenCV(CharucoBoard->ArucoDictionary);
+
+	// Detect Charuco corners
+	FCharucoCorners DetectedCorners;
+	bool bDetectionSuccess = FOpenCVHelper::IdentifyCharucoCorners(Pixels, Size, CachedBoardConfig, DetectedCorners, ROI);
+
+	if (!bDetectionSuccess || DetectedCorners.Corners.IsEmpty())
+	{
+		OutErrorMessage = FText::FromString(FString::Printf(TEXT(
+			"Could not identify Charuco corners. "
+			"The expected board has %dx%d squares with ArUco dictionary %s. "
+			"Make sure the board is clearly visible and well-lit."),
+			CachedBoardConfig.SquaresX, CachedBoardConfig.SquaresY,
+			*UEnum::GetValueAsString(CachedBoardConfig.Dictionary)
+		));
+		return false;
+	}
+
+	// Show detection visualization if requested
+	if (bShouldShowDetectionWindow)
+	{
+		ShowDetectionVisualization(DetectedCorners, Pixels, Size, ROI);
+	}
+
+	// Export the latest session data
+	ExportSessionData();
+
+	if (!ensure(DetectedCorners.Corners.Num() == DetectedCorners.CornerIds.Num()))
+	{
+		OutErrorMessage = FText::Format(LOCTEXT("CharucoArrayMismatch", "Charuco detection failed: corner count ({0}) does not match corner ID count ({1})"), 
+			DetectedCorners.Corners.Num(), DetectedCorners.CornerIds.Num());
+		return false;
+	}
+
+	// Create calibration rows for each detected corner
+	for (int32 CornerIdx = 0; CornerIdx < DetectedCorners.Corners.Num(); ++CornerIdx)
+	{
+		int32 CornerId = DetectedCorners.CornerIds[CornerIdx];
+		FVector2f ImageCorner = DetectedCorners.Corners[CornerIdx];
+		
+		// Calculate 3D position using same approach as checkerboard algorithm
+		int32 CornerRow = CornerId / CharucoBoard->NumCornerCols;
+		int32 CornerCol = CornerId % CharucoBoard->NumCornerCols;
+		
+		FTransform LocalPoint3d;
+
+		// Calculate local position based on the corner coordinates and the square side length
+		const FVector LocalPosition = CharucoBoard->SquareSideLength * FVector(0, CornerCol + 1, CornerRow + 1);
+		LocalPoint3d.SetLocation(LocalPosition);
+		
+		const FVector WorldCorner = (LocalPoint3d * Calibrator->GetTransform()).GetLocation();
+
+		TSharedPtr<FNodalOffsetPointsRowData> Row = MakeShared<FNodalOffsetPointsRowData>();
+
+		// Convert to the expected format, normalize by image size.
+		Row->Point2D = FVector2f(ImageCorner.X / Size.X, ImageCorner.Y / Size.Y);
+		Row->CameraData = LastCameraData;
+
+		// Set the calibrator point data
+		Row->CalibratorPointData.bIsValid = true;
+		Row->CalibratorPointData.Name = FString::Printf(TEXT("Corner_%d"), CornerId);
+		Row->CalibratorPointData.Location = WorldCorner;
+
+		// Add to the calibration points
+		CalibrationRows.Add(Row);
+
+		// Export the data for this row to a .json file on disk
+		ExportRow(Row);
+	}
+
+	// Notify the ListView of the new data
+	if (CalibrationListView.IsValid())
+	{
+		CalibrationListView->RequestListRefresh();
+	}
+
+	return true;
+}
+
+void UCameraNodalOffsetAlgoCharuco::ShowDetectionVisualization(const FCharucoCorners& DetectedCorners, const TArray<FColor>& Image, FIntPoint ImageSize, const FIntRect& ROI)
+{
+	// Create debug texture
+	UTexture2D* DebugTexture = UTexture2D::CreateTransient(ImageSize.X, ImageSize.Y, EPixelFormat::PF_B8G8R8A8);
+
+	if (!DebugTexture)
+	{
+		return;
+	}
+
+	// Copy image data to texture in nested scope to ensure unlock before UpdateResource
+	{
+		FTexture2DMipMap& Mip = DebugTexture->GetPlatformData()->Mips[0];
+		uint8* TextureData = static_cast<uint8*>(Mip.BulkData.Lock(LOCK_READ_WRITE));
+		
+		ON_SCOPE_EXIT
+		{
+			if (Mip.BulkData.IsLocked())
+			{
+				Mip.BulkData.Unlock();
+			}
+		};
+		
+		if (!TextureData)
+		{
+			DebugTexture->MarkAsGarbage();
+			return;
+		}
+		
+		FMemory::Memcpy(TextureData, Image.GetData(), Image.Num() * sizeof(FColor));
+	}
+
+	DebugTexture->UpdateResource();
+
+	// Draw the detection results on the texture
+	FOpenCVHelper::DrawCharucoCorners(DetectedCorners, DebugTexture);
+
+	// Draw ROI rectangle if specified
+	if (ROI.Width() > 0 && ROI.Height() > 0)
+	{
+		DrawROIRectangle(DebugTexture, ROI, ImageSize);
+	}
+
+	// Display the texture in a window
+	FText WindowTitle = ROI.Width() > 0 && ROI.Height() > 0 ? 
+		FText::Format(LOCTEXT("CharucoDetectionROI", "Charuco Corner Detection ROI ({0},{1}) to ({2},{3})"), 
+			ROI.Min.X, ROI.Min.Y, ROI.Max.X, ROI.Max.Y) :
+		LOCTEXT("CharucoDetection", "Charuco Corner Detection");
+		
+	FCameraCalibrationWidgetHelpers::DisplayTextureInWindowAlmostFullScreen(DebugTexture, MoveTemp(WindowTitle));
+}
+
+void UCameraNodalOffsetAlgoCharuco::DrawROIRectangle(UTexture2D* DebugTexture, const FIntRect& ROI, FIntPoint ImageSize)
+{
+	if (!DebugTexture)
+	{
+		return;
+	}
+
+	// Draw ROI rectangle
+	{
+		FTexture2DMipMap& Mip = DebugTexture->GetPlatformData()->Mips[0];
+		uint8* TextureData = static_cast<uint8*>(Mip.BulkData.Lock(LOCK_READ_WRITE));
+		
+		ON_SCOPE_EXIT
+		{
+			if (Mip.BulkData.IsLocked())
+			{
+				Mip.BulkData.Unlock();
+			}
+		};
+		
+		if (!TextureData)
+		{
+			return;
+		}
+
+		// ROI rectangle color (bright cyan for good visibility)
+		const FColor ROIColor = FColor::Cyan;
+		const int32 LineThickness = 3; // Make it thick enough to see
+
+		// Clamp ROI to image bounds
+		const int32 MinX = FMath::Clamp(ROI.Min.X, 0, ImageSize.X - 1);
+		const int32 MinY = FMath::Clamp(ROI.Min.Y, 0, ImageSize.Y - 1);
+		const int32 MaxX = FMath::Clamp(ROI.Max.X, MinX + 1, ImageSize.X);
+		const int32 MaxY = FMath::Clamp(ROI.Max.Y, MinY + 1, ImageSize.Y);
+
+		// Draw horizontal lines (top and bottom)
+		for (int32 Thickness = 0; Thickness < LineThickness; ++Thickness)
+		{
+			// Top line
+			int32 TopY = FMath::Clamp(MinY + Thickness, 0, ImageSize.Y - 1);
+			for (int32 X = MinX; X < MaxX; ++X)
+			{
+				int32 PixelIndex = TopY * ImageSize.X + X;
+				reinterpret_cast<FColor*>(TextureData)[PixelIndex] = ROIColor;
+			}
+
+			// Bottom line  
+			int32 BottomY = FMath::Clamp(MaxY - 1 - Thickness, 0, ImageSize.Y - 1);
+			for (int32 X = MinX; X < MaxX; ++X)
+			{
+				int32 PixelIndex = BottomY * ImageSize.X + X;
+				reinterpret_cast<FColor*>(TextureData)[PixelIndex] = ROIColor;
+			}
+		}
+
+		// Draw vertical lines (left and right)
+		for (int32 Thickness = 0; Thickness < LineThickness; ++Thickness)
+		{
+			// Left line
+			int32 LeftX = FMath::Clamp(MinX + Thickness, 0, ImageSize.X - 1);
+			for (int32 Y = MinY; Y < MaxY; ++Y)
+			{
+				int32 PixelIndex = Y * ImageSize.X + LeftX;
+				reinterpret_cast<FColor*>(TextureData)[PixelIndex] = ROIColor;
+			}
+
+			// Right line
+			int32 RightX = FMath::Clamp(MaxX - 1 - Thickness, 0, ImageSize.X - 1);
+			for (int32 Y = MinY; Y < MaxY; ++Y)
+			{
+				int32 PixelIndex = Y * ImageSize.X + RightX;
+				reinterpret_cast<FColor*>(TextureData)[PixelIndex] = ROIColor;
+			}
+		}
+	}
+
+	DebugTexture->UpdateResource();
+}
+
+TSharedRef<SWidget> UCameraNodalOffsetAlgoCharuco::BuildHelpWidget()
+{
+	return SNew(STextBlock)
+		.Text(LOCTEXT("CharucoHelp", 
+			"This algorithm uses a Charuco board (chessboard with embedded ArUco markers) for calibration.\n\n"
+			"Advantages:\n"
+			"• Higher precision than pure ArUco markers\n"
+			"• More robust than pure checkerboards\n"
+			"• Partial occlusion handling\n"
+			"Instructions:\n"
+			"1. Create a Charuco board actor in your scene\n"
+			"2. Select the board in the 'Charuco Board' dropdown\n"
+			"3. Position the board so it's clearly visible in the media feed\n"
+			"4. Click on the media viewport to capture calibration points\n"
+			"5. Repeat from different positions and orientations"
+		))
+		.AutoWrapText(true);
+}
+
+#undef LOCTEXT_NAMESPACE
\ No newline at end of file
diff --git a/Engine/Plugins/VirtualProduction/CameraCalibration/Source/CameraCalibrationEditor/Private/Calibrators/CameraNodalOffsetAlgoCharuco.h b/Engine/Plugins/VirtualProduction/CameraCalibration/Source/CameraCalibrationEditor/Private/Calibrators/CameraNodalOffsetAlgoCharuco.h
new file mode 100644
--- /dev/null
+++ b/Engine/Plugins/VirtualProduction/CameraCalibration/Source/CameraCalibrationEditor/Private/Calibrators/CameraNodalOffsetAlgoCharuco.h
@@ -0,0 +1,70 @@
+// Copyright Epic Games, Inc. All Rights Reserved.
+
+#pragma once
+
+#include "CameraCalibrationCharucoBoard.h" // For shadow types
+#include "CameraNodalOffsetAlgoPoints.h"
+#include "OpenCVHelper.h"
+
+#include "CameraNodalOffsetAlgoCharuco.generated.h"
+
+class ACameraCalibrationCharucoBoard;
+class SWidget;
+
+/** 
+ * Implements a nodal offset calibration algorithm based on a Charuco board.
+ * Charuco combines the benefits of checkerboard corners (high precision) 
+ * with ArUco markers (robust detection and orientation).
+ */
+UCLASS()
+class UCameraNodalOffsetAlgoCharuco : public UCameraNodalOffsetAlgoPoints
+{
+	GENERATED_BODY()
+	
+public:
+
+	//~ Begin UCameraNodalOffsetAlgoPoints Interface
+	virtual bool OnViewportClicked(const FGeometry& MyGeometry, const FPointerEvent& MouseEvent) override;
+	virtual bool OnViewportMarqueeSelect(FVector2D StartPosition, FVector2D EndPosition) override;
+	virtual TSharedRef<SWidget> BuildUI() override;
+	virtual FName FriendlyName() const override { return TEXT("Nodal Offset Charuco"); }
+	virtual FName ShortName() const override { return TEXT("Charuco"); }
+	virtual TSharedRef<SWidget> BuildHelpWidget() override;
+	//~ End UCameraNodalOffsetAlgoPoints Interface
+
+
+protected:
+
+	//~ Begin UCameraNodalOffsetAlgoPoints Interface
+	virtual AActor* FindFirstCalibrator() const override;
+	//~ End UCameraNodalOffsetAlgoPoints Interface
+
+protected:
+
+	/** Builds the UI for the Charuco board picker */
+	TSharedRef<SWidget> BuildCharucoBoardPickerWidget();
+
+	/** Builds the UI for showing/hiding detection visualization */
+	TSharedRef<SWidget> BuildShowDetectionWidget();
+
+	/** Populates the calibration rows with detected Charuco corners. True if successful. */
+	virtual bool PopulatePoints(FText& OutErrorMessage);
+	
+	/** Populates the calibration rows with detected Charuco corners using ROI. True if successful. */
+	virtual bool PopulatePoints(FText& OutErrorMessage, const FIntRect& ROI);
+
+	/** Shows a debug window with the detected corners and markers, optionally showing ROI */
+	void ShowDetectionVisualization(const FCharucoCorners& DetectedCorners, const TArray<FColor>& Image, FIntPoint ImageSize, const FIntRect& ROI = FIntRect());
+
+	/** Draws an ROI rectangle on the debug texture */
+	void DrawROIRectangle(UTexture2D* DebugTexture, const FIntRect& ROI, FIntPoint ImageSize);
+
+protected:
+
+	/** True if a detection window should be shown after every capture */
+	UPROPERTY()
+	bool bShouldShowDetectionWindow = false;
+
+	/** Cached Charuco board configuration from the selected board actor */
+	FCharucoBoardConfig CachedBoardConfig;
+};
\ No newline at end of file
diff --git a/Engine/Plugins/VirtualProduction/CameraCalibration/Source/CameraCalibrationEditor/Private/CameraCalibrationUtilsPrivate.cpp b/Engine/Plugins/VirtualProduction/CameraCalibration/Source/CameraCalibrationEditor/Private/CameraCalibrationUtilsPrivate.cpp
--- a/Engine/Plugins/VirtualProduction/CameraCalibration/Source/CameraCalibrationEditor/Private/CameraCalibrationUtilsPrivate.cpp
+++ b/Engine/Plugins/VirtualProduction/CameraCalibration/Source/CameraCalibrationEditor/Private/CameraCalibrationUtilsPrivate.cpp
@@ -181,11 +181,31 @@ namespace UE::CameraCalibration::Private
 			TArray<FColor> Pixels;
 			Pixels.Init(ClearColor, Texture->GetSizeX() * Texture->GetSizeY());
 
-			void* TextureData = Texture->GetPlatformData()->Mips[0].BulkData.Lock(LOCK_READ_WRITE);
+			{
+				FTexture2DMipMap& Mip = Texture->GetPlatformData()->Mips[0];
+
+				const SIZE_T SrcBytes = SIZE_T(Pixels.Num()) * SIZE_T(sizeof(FColor));
+				const SIZE_T DstBytes = SIZE_T(Mip.BulkData.GetBulkDataSize());
+
+				const SIZE_T NumBytes = FMath::Min(SrcBytes, DstBytes);
+
+				void* TextureData = Mip.BulkData.Lock(LOCK_READ_WRITE);
+
+				ON_SCOPE_EXIT
+				{
+					if (Mip.BulkData.IsLocked())
+					{
+						Mip.BulkData.Unlock();
+					}
+				};
 
-			FMemory::Memcpy(TextureData, Pixels.GetData(), Pixels.Num() * sizeof(FColor));
+				if (!TextureData)
+				{
+					return;
+				}
 
-			Texture->GetPlatformData()->Mips[0].BulkData.Unlock();
+				FMemory::Memcpy(TextureData, Pixels.GetData(), NumBytes);
+			}
 			Texture->UpdateResource();
 		}
 	}
@@ -194,11 +214,32 @@ namespace UE::CameraCalibration::Private
 	{
 		if (Texture)
 		{
-			void* TextureData = Texture->GetPlatformData()->Mips[0].BulkData.Lock(LOCK_READ_WRITE);
+			{
+				FTexture2DMipMap& Mip = Texture->GetPlatformData()->Mips[0];
 
-			FMemory::Memcpy(TextureData, PixelData.GetData(), PixelData.Num() * sizeof(FColor));
+				const SIZE_T SrcBytes = SIZE_T(PixelData.Num()) * SIZE_T(sizeof(FColor));
+				const SIZE_T DstBytes = SIZE_T(Mip.BulkData.GetBulkDataSize());
+
+				const SIZE_T NumBytes = FMath::Min(SrcBytes, DstBytes);
+
+				void* TextureData = Mip.BulkData.Lock(LOCK_READ_WRITE);
+
+				ON_SCOPE_EXIT
+				{
+					if (Mip.BulkData.IsLocked())
+					{
+						Mip.BulkData.Unlock();
+					}
+				};
+
+				if (!TextureData)
+				{
+					return;
+				}
+
+				FMemory::Memcpy(TextureData, PixelData.GetData(), NumBytes);
+			}
 
-			Texture->GetPlatformData()->Mips[0].BulkData.Unlock();
 			Texture->UpdateResource();
 		}
 	}
diff --git a/Engine/Plugins/VirtualProduction/CameraCalibrationCore/Source/CameraCalibrationCore/CameraCalibrationCore.Build.cs b/Engine/Plugins/VirtualProduction/CameraCalibrationCore/Source/CameraCalibrationCore/CameraCalibrationCore.Build.cs
--- a/Engine/Plugins/VirtualProduction/CameraCalibrationCore/Source/CameraCalibrationCore/CameraCalibrationCore.Build.cs
+++ b/Engine/Plugins/VirtualProduction/CameraCalibrationCore/Source/CameraCalibrationCore/CameraCalibrationCore.Build.cs
@@ -6,6 +6,8 @@ namespace UnrealBuildTool.Rules
 	{
 		public CameraCalibrationCore(ReadOnlyTargetRules Target) : base(Target)
 		{
+			// Note: For now, avoid adding OpenCVHelper/OpenCV as a dependency until a server dependency issue is resolved.
+
 			PublicDependencyModuleNames.AddRange(
 				new string[]
 				{
diff --git a/Engine/Plugins/VirtualProduction/CameraCalibrationCore/Source/CameraCalibrationCore/Private/CameraCalibrationCharucoBoard.cpp b/Engine/Plugins/VirtualProduction/CameraCalibrationCore/Source/CameraCalibrationCore/Private/CameraCalibrationCharucoBoard.cpp
--- a/Engine/Plugins/VirtualProduction/CameraCalibrationCore/Source/CameraCalibrationCore/Private/CameraCalibrationCharucoBoard.cpp
+++ b/Engine/Plugins/VirtualProduction/CameraCalibrationCore/Source/CameraCalibrationCore/Private/CameraCalibrationCharucoBoard.cpp
@@ -91,8 +91,7 @@ void ACameraCalibrationCharucoBoard::PostEditChangeProperty(FPropertyChangedEven
 			PropertyName == GET_MEMBER_NAME_CHECKED(ACameraCalibrationCharucoBoard, Thickness) ||
 			PropertyName == GET_MEMBER_NAME_CHECKED(ACameraCalibrationCharucoBoard, BoardMesh) ||
 			PropertyName == GET_MEMBER_NAME_CHECKED(ACameraCalibrationCharucoBoard, BoardMaterial) ||
-			PropertyName == GET_MEMBER_NAME_CHECKED(ACameraCalibrationCharucoBoard, MaxTextureResolution) ||
-			PropertyName == GET_MEMBER_NAME_CHECKED(ACameraCalibrationCharucoBoard, bDisplayPattern))
+			PropertyName == GET_MEMBER_NAME_CHECKED(ACameraCalibrationCharucoBoard, MaxTextureResolution))
 		{
 			Rebuild();
 		}
@@ -141,8 +140,8 @@ void ACameraCalibrationCharucoBoard::CreateBoardMesh()
 	// The cube mesh has its origin at center, so offset by half the scaled dimensions  
 	BoardMeshComponent->SetRelativeLocation(FVector(0.0f, BoardWidth / 2.0f, BoardHeight / 2.0f));
 
-	// Apply material if pattern display is enabled
-	if (bDisplayPattern && BoardMaterial)
+	// Apply material for pattern display
+	if (BoardMaterial)
 	{
 		DynamicBoardMaterial = UMaterialInstanceDynamic::Create(BoardMaterial, this);
 		BoardMeshComponent->SetMaterial(0, DynamicBoardMaterial);
@@ -200,16 +199,15 @@ void ACameraCalibrationCharucoBoard::GenerateCharucoTexture()
 	}
 
 	// Apply texture to material if available
-	if (CharucoBoardTexture && DynamicBoardMaterial && bDisplayPattern)
+	if (CharucoBoardTexture && DynamicBoardMaterial)
 	{
 		DynamicBoardMaterial->SetTextureParameterValue(TEXT("Texture"), CharucoBoardTexture);
 	}
 	else
 	{
-		UE_LOG(LogCameraCalibrationCore, Error, TEXT("Cannot apply texture - CharucoBoardTexture: %s, DynamicBoardMaterial: %s, bDisplayPattern: %s"),
+		UE_LOG(LogCameraCalibrationCore, Error, TEXT("Cannot apply texture - CharucoBoardTexture: %s, DynamicBoardMaterial: %s"),
 			CharucoBoardTexture ? TEXT("Valid") : TEXT("NULL"),
-			DynamicBoardMaterial ? TEXT("Valid") : TEXT("NULL"),
-			bDisplayPattern ? TEXT("True") : TEXT("False"));
+			DynamicBoardMaterial ? TEXT("Valid") : TEXT("NULL"));
 	}
 }
 
diff --git a/Engine/Plugins/VirtualProduction/CameraCalibrationCore/Source/CameraCalibrationCore/Public/CameraCalibrationCharucoBoard.h b/Engine/Plugins/VirtualProduction/CameraCalibrationCore/Source/CameraCalibrationCore/Public/CameraCalibrationCharucoBoard.h
--- a/Engine/Plugins/VirtualProduction/CameraCalibrationCore/Source/CameraCalibrationCore/Public/CameraCalibrationCharucoBoard.h
+++ b/Engine/Plugins/VirtualProduction/CameraCalibrationCore/Source/CameraCalibrationCore/Public/CameraCalibrationCharucoBoard.h
@@ -176,10 +176,6 @@ public:
 	UPROPERTY(EditAnywhere, Category = "Calibration", meta = (ClampMin = "256", ClampMax = "4096", UIMin = "256", UIMax = "4096"))
 	int32 MaxTextureResolution = 2048;
 
-	/** Display the generated Charuco pattern on the board mesh */
-	UPROPERTY(EditAnywhere, Category = "Calibration")
-	bool bDisplayPattern = true;
-
 private:
 	/** Generated texture containing the Charuco pattern */
 	UPROPERTY(Transient)
