CameraCalibration: Went back to using local space coordinates when calibrating the charuco.

I found cases where using world coordinates yields inaccurate results with both the legacy and the custom solver.

#rnx

diff --git a/Engine/Plugins/Runtime/OpenCV/Source/OpenCVHelper/Private/OpenCVHelper.cpp b/Engine/Plugins/Runtime/OpenCV/Source/OpenCVHelper/Private/OpenCVHelper.cpp
--- a/Engine/Plugins/Runtime/OpenCV/Source/OpenCVHelper/Private/OpenCVHelper.cpp
+++ b/Engine/Plugins/Runtime/OpenCV/Source/OpenCVHelper/Private/OpenCVHelper.cpp
@@ -42,7 +42,7 @@ namespace
 		if (RequiredBytes > DestBufferSize)
 		{
 			UE_LOG(LogOpenCVHelper, Error, TEXT("SafeCopyMatToBuffer: Destination buffer too small. Required: %llu bytes, Available: %llu bytes"), 
-				RequiredBytes, DestBufferSize);
+				static_cast<uint64>(RequiredBytes), static_cast<uint64>(DestBufferSize));
 			return false;
 		}
 		
diff --git a/Engine/Plugins/VirtualProduction/CameraCalibration/Source/CameraCalibrationEditor/Private/AssetEditor/LensDistortionTool.cpp b/Engine/Plugins/VirtualProduction/CameraCalibration/Source/CameraCalibrationEditor/Private/AssetEditor/LensDistortionTool.cpp
--- a/Engine/Plugins/VirtualProduction/CameraCalibration/Source/CameraCalibrationEditor/Private/AssetEditor/LensDistortionTool.cpp
+++ b/Engine/Plugins/VirtualProduction/CameraCalibration/Source/CameraCalibrationEditor/Private/AssetEditor/LensDistortionTool.cpp
@@ -36,10 +36,10 @@
 // Console variable to control whether to generate object points in local space for Charuco calibration
 static TAutoConsoleVariable<bool> CVarCharucoUseLocalSpaceObjectPoints(
 	TEXT("CameraCalibration.Charuco.UseLocalSpaceObjectPoints"),
-	false,
+	true,
 	TEXT("If true, generate object points in local space for Charuco calibration instead of using world space points.\n")
-	TEXT("The custom solver should not need object points in a canonical plane anymore but leaving this option here while it gets battle-testd\n")
-	TEXT("Default: false"),
+	TEXT("The solvers don't require object points in a canonical plane anymore, but can yield less accurate results otherwise.\n")
+	TEXT("Default: true"),
 	ECVF_Default
 );
 
@@ -790,7 +790,7 @@ void ULensDistortionTool::CalibrateLens()
 	}
 	else if (CaptureSettings.CalibrationPattern == ECalibrationPattern::Charuco)
 	{
-		// Custom solver should not need object points in a canonical plane anymore, but leaving this here while this is battle-tested.
+		// The solvers don't require object points in a canonical plane anymore, but can yield less accurate results otherwise.
 		if (CVarCharucoUseLocalSpaceObjectPoints.GetValueOnGameThread())
 		{
 			GenerateObjectPointsInLocalSpace(Samples3d);
diff --git a/Engine/Plugins/VirtualProduction/CameraCalibration/Source/CameraCalibrationEditor/Private/Calibrators/CameraCalibrationSolver.cpp b/Engine/Plugins/VirtualProduction/CameraCalibration/Source/CameraCalibrationEditor/Private/Calibrators/CameraCalibrationSolver.cpp
--- a/Engine/Plugins/VirtualProduction/CameraCalibration/Source/CameraCalibrationEditor/Private/Calibrators/CameraCalibrationSolver.cpp
+++ b/Engine/Plugins/VirtualProduction/CameraCalibration/Source/CameraCalibrationEditor/Private/Calibrators/CameraCalibrationSolver.cpp
@@ -705,11 +705,15 @@ void ULensDistortionSolverOpenCV::InitIntrinsicsFromCanonicalPlane(
 	const double Cx = (ImageSize.width - 1) * 0.5;
 	const double Cy = (ImageSize.height - 1) * 0.5;
 
-	cv::Mat MatrixA = cv::Mat(2 * NumImages, 2, CV_64F);
-	cv::Mat VecB = cv::Mat(2 * NumImages, 1, CV_64F);
+	// Collect valid homography constraint rows to avoid uninitialized matrix data
+	std::vector<double> ValidMatrixARows;
+	std::vector<double> ValidVecBRows;
+	ValidMatrixARows.reserve(2 * NumImages * 2); // Each image contributes 2 rows with 2 columns
+	ValidVecBRows.reserve(2 * NumImages);
 
 	// extract vanishing points in order to obtain initial value for the focal length
 	int ObjectPointsIndex = 0;
+	int ValidImageCount = 0;
 	for (int ImageIndex = 0; ImageIndex < NumImages; ++ImageIndex)
 	{
 		const int NumImagePoints = NumPoints.ptr<int>()[ImageIndex];
@@ -721,6 +725,12 @@ void ULensDistortionSolverOpenCV::InitIntrinsicsFromCanonicalPlane(
 
 		cv::Mat HomographyMat = cv::findHomography(ObjectPointsInImage, ImagePointsInImage);
 
+		if (HomographyMat.empty() || HomographyMat.rows != 3 || HomographyMat.cols != 3 || HomographyMat.type() != CV_64F)
+		{
+			UE_LOG(LogCameraCalibrationSolver, Warning, TEXT("findHomography failed or returned invalid matrix for image %d"), ImageIndex);
+			continue;
+		}
+
 		double* Homography = HomographyMat.ptr<double>();
 
 		Homography[0] -= Homography[6] * Cx;
@@ -767,18 +777,38 @@ void ULensDistortionSolverOpenCV::InitIntrinsicsFromCanonicalPlane(
 			D2[Index] *= N[3];
 		}
 
-		double* MatrixAPtr = MatrixA.ptr<double>(ImageIndex * 2);
-		double* VecBPtr = VecB.ptr<double>(ImageIndex * 2);
+		// Store valid constraint rows
+		ValidMatrixARows.push_back(H[0] * V[0]);
+		ValidMatrixARows.push_back(H[1] * V[1]);
+		ValidMatrixARows.push_back(D1[0] * D2[0]);
+		ValidMatrixARows.push_back(D1[1] * D2[1]);
+
+		ValidVecBRows.push_back(-H[2] * V[2]);
+		ValidVecBRows.push_back(-D1[2] * D2[2]);
+
+		ValidImageCount++;
+	}
+
+	// Check if we have enough valid homographies for focal length estimation
+	if (ValidImageCount < 2)
+	{
+		UE_LOG(LogCameraCalibrationSolver, Warning, TEXT("InitIntrinsicsFromCanonicalPlane: Insufficient valid homographies (%d) for focal length estimation. Need at least 2."), ValidImageCount);
 
-		MatrixAPtr[0] = H[0] * V[0];
-		MatrixAPtr[1] = H[1] * V[1];
-		MatrixAPtr[2] = D1[0] * D2[0];
-		MatrixAPtr[3] = D1[1] * D2[1];
+		// Fallback to basic initialization
+		const double Fx = ImageSize.width;
+		const double Fy = ImageSize.width;
 
-		VecBPtr[0] = -H[2] * V[2];
-		VecBPtr[1] = -D1[2] * D2[2];
+		CameraMatrix.at<double>(0, 0) = Fx;
+		CameraMatrix.at<double>(1, 1) = Fy;
+		CameraMatrix.at<double>(0, 2) = Cx;
+		CameraMatrix.at<double>(1, 2) = Cy;
+		return;
 	}
 
+	// Build matrices only with valid data
+	cv::Mat MatrixA = cv::Mat(2 * ValidImageCount, 2, CV_64F, ValidMatrixARows.data());
+	cv::Mat VecB = cv::Mat(2 * ValidImageCount, 1, CV_64F, ValidVecBRows.data());
+
 	double FxFy[2] = { 0 };
 	cv::Mat FxFyMat = cv::Mat(2, 1, CV_64F, FxFy);
 
diff --git a/Engine/Plugins/VirtualProduction/CameraCalibration/Source/CameraCalibrationEditor/Private/Calibrators/CameraNodalOffsetAlgoCharuco.h b/Engine/Plugins/VirtualProduction/CameraCalibration/Source/CameraCalibrationEditor/Private/Calibrators/CameraNodalOffsetAlgoCharuco.h
--- a/Engine/Plugins/VirtualProduction/CameraCalibration/Source/CameraCalibrationEditor/Private/Calibrators/CameraNodalOffsetAlgoCharuco.h
+++ b/Engine/Plugins/VirtualProduction/CameraCalibration/Source/CameraCalibrationEditor/Private/Calibrators/CameraNodalOffsetAlgoCharuco.h
@@ -2,7 +2,7 @@
 
 #pragma once
 
-#include "CameraCalibrationCharucoBoard.h" // For shadow types
+#include "CameraCalibrationCharucoBoard.h"
 #include "CameraNodalOffsetAlgoPoints.h"
 #include "OpenCVHelper.h"
 
